Add world readiness and world readiness test modules

diff --git a/Engine/Plugins/Experimental/SpatialReadiness/Source/PhysicsSpatialReadiness/PhysicsSpatialReadiness.Build.cs b/Engine/Plugins/Experimental/SpatialReadiness/Source/PhysicsSpatialReadiness/PhysicsSpatialReadiness.Build.cs
new file mode 100644
--- /dev/null
+++ b/Engine/Plugins/Experimental/SpatialReadiness/Source/PhysicsSpatialReadiness/PhysicsSpatialReadiness.Build.cs
@@ -0,0 +1,24 @@
+// Copyright Epic Games, Inc. All Rights Reserved.
+
+using UnrealBuildTool;
+
+public class PhysicsSpatialReadiness : ModuleRules
+{
+    public PhysicsSpatialReadiness(ReadOnlyTargetRules Target) : base(Target)
+    {
+        PrivateDependencyModuleNames.AddRange(
+            new string[]
+            {
+                "Engine",
+                "SpatialReadiness",
+            });
+
+        PublicDependencyModuleNames.AddRange(
+            new string[] {
+                "Core",
+            }
+        );
+
+        SetupModulePhysicsSupport(Target);
+    }
+}
diff --git a/Engine/Plugins/Experimental/SpatialReadiness/Source/PhysicsSpatialReadiness/Private/PhysicsSpatialReadiness/PhysicsSpatialReadinessModule.cpp b/Engine/Plugins/Experimental/SpatialReadiness/Source/PhysicsSpatialReadiness/Private/PhysicsSpatialReadiness/PhysicsSpatialReadinessModule.cpp
new file mode 100644
--- /dev/null
+++ b/Engine/Plugins/Experimental/SpatialReadiness/Source/PhysicsSpatialReadiness/Private/PhysicsSpatialReadiness/PhysicsSpatialReadinessModule.cpp
@@ -0,0 +1,10 @@
+// Copyright Epic Games, Inc. All Rights Reserved.
+
+#include "PhysicsSpatialReadinessModule.h"
+#include "Modules/ModuleManager.h"
+
+IMPLEMENT_MODULE(FDefaultModuleImpl, PhysicsSpatialReadiness)
+
+void FPhysicsSpatialReadinessModule::StartupModule() { }
+
+void FPhysicsSpatialReadinessModule::ShutdownModule() { }
\ No newline at end of file
diff --git a/Engine/Plugins/Experimental/SpatialReadiness/Source/PhysicsSpatialReadiness/Private/PhysicsSpatialReadiness/PhysicsSpatialReadinessSimCallback.cpp b/Engine/Plugins/Experimental/SpatialReadiness/Source/PhysicsSpatialReadiness/Private/PhysicsSpatialReadiness/PhysicsSpatialReadinessSimCallback.cpp
new file mode 100644
--- /dev/null
+++ b/Engine/Plugins/Experimental/SpatialReadiness/Source/PhysicsSpatialReadiness/Private/PhysicsSpatialReadiness/PhysicsSpatialReadinessSimCallback.cpp
@@ -0,0 +1,17 @@
+// Copyright Epic Games, Inc. All Rights Reserved.
+
+#include "PhysicsSpatialReadinessSimCallback.h"
+
+using namespace Chaos;
+
+void FPhysicsSpatialReadinessSimCallback::OnPreSimulate_Internal()
+{
+}
+
+void FPhysicsSpatialReadinessSimCallback::OnMidPhaseModification_Internal(FMidPhaseModifierAccessor& Accessor)
+{
+}
+
+void FPhysicsSpatialReadinessSimCallback::OnPostIntegrate_Internal()
+{
+}
diff --git a/Engine/Plugins/Experimental/SpatialReadiness/Source/PhysicsSpatialReadiness/Private/PhysicsSpatialReadiness/PhysicsSpatialReadinessSubsystem.cpp b/Engine/Plugins/Experimental/SpatialReadiness/Source/PhysicsSpatialReadiness/Private/PhysicsSpatialReadiness/PhysicsSpatialReadinessSubsystem.cpp
new file mode 100644
--- /dev/null
+++ b/Engine/Plugins/Experimental/SpatialReadiness/Source/PhysicsSpatialReadiness/Private/PhysicsSpatialReadiness/PhysicsSpatialReadinessSubsystem.cpp
@@ -0,0 +1,98 @@
+// Copyright Epic Games, Inc. All Rights Reserved.
+
+#include "PhysicsSpatialReadinessSubsystem.h"
+#include "Engine/World.h"
+#include "Physics/Experimental/PhysScene_Chaos.h"
+#include "PBDRigidsSolver.h"
+
+using namespace Chaos;
+
+bool UPhysicsSpatialReadinessSubsystem::ShouldCreateSubsystem(UObject* Outer) const
+{
+	// TODO: Where should the setting for this exist? Is it enough
+	//       to just not load the physics readiness module if we don't
+	//       want to use it?
+	return true;
+}
+
+void UPhysicsSpatialReadinessSubsystem::OnWorldBeginPlay(UWorld& InWorld)
+{
+	Super::OnWorldBeginPlay(InWorld);
+	CreateSimCallback();
+}
+
+void UPhysicsSpatialReadinessSubsystem::Initialize(FSubsystemCollectionBase& Collection)
+{
+	Super::Initialize(Collection);
+}
+
+void UPhysicsSpatialReadinessSubsystem::Deinitialize()
+{
+	DestroySimCallback();
+	Super::Deinitialize();
+}
+
+FUnreadyVolumeHandle UPhysicsSpatialReadinessSubsystem::AddUnreadyVolume(const FBox& AABB, const FString& Description)
+{
+	/*
+	* TODO:
+	* - Create geometry particle for unready volume
+	* - Add particle to list of unready particles in sim callback
+	* - Add particle to our internal array of unready particles at first free index
+	* - Return the index that the particle was stored at as the "handle" to the volume
+	*/
+	return FUnreadyVolumeHandle();
+}
+
+void UPhysicsSpatialReadinessSubsystem::RemoveUnreadyVolume(FUnreadyVolumeHandle Handle)
+{
+	/*
+	* TODO:
+	* - Destroy geometry particle corresponding to volume handle
+	* - Free index
+	*/
+}
+
+bool UPhysicsSpatialReadinessSubsystem::CreateSimCallback()
+{
+	// In case we already have one, destroy it
+	DestroySimCallback();
+
+	if (SimCallback == nullptr)
+	{
+		if (UWorld* World = GetWorld())
+		{
+			if (FPhysScene* PhysScene = World->GetPhysicsScene())
+			{
+				if (FPhysicsSolverBase* Solver = PhysScene->GetSolver())
+				{
+					SimCallback = Solver->CreateAndRegisterSimCallbackObject_External<FPhysicsSpatialReadinessSimCallback>();
+					return true;
+				}
+			}
+		}
+	}
+
+	return false;
+}
+
+bool UPhysicsSpatialReadinessSubsystem::DestroySimCallback()
+{
+	if (SimCallback != nullptr)
+	{
+		if (UWorld* World = GetWorld())
+		{
+			if (FPhysScene* PhysScene = World->GetPhysicsScene())
+			{
+				if (FPhysicsSolverBase* Solver = PhysScene->GetSolver())
+				{
+					Solver->UnregisterAndFreeSimCallbackObject_External(SimCallback);
+					SimCallback = nullptr;
+					return true;
+				}
+			}
+		}
+	}
+
+	return false;
+}
diff --git a/Engine/Plugins/Experimental/SpatialReadiness/Source/PhysicsSpatialReadiness/Public/PhysicsSpatialReadinessModule.h b/Engine/Plugins/Experimental/SpatialReadiness/Source/PhysicsSpatialReadiness/Public/PhysicsSpatialReadinessModule.h
new file mode 100644
--- /dev/null
+++ b/Engine/Plugins/Experimental/SpatialReadiness/Source/PhysicsSpatialReadiness/Public/PhysicsSpatialReadinessModule.h
@@ -0,0 +1,13 @@
+// Copyright Epic Games, Inc. All Rights Reserved.
+
+#pragma once
+
+#include "CoreMinimal.h"
+#include "Modules/ModuleManager.h"
+
+class FPhysicsSpatialReadinessModule : public IModuleInterface
+{
+public:
+	virtual void StartupModule() override;
+	virtual void ShutdownModule() override;
+};
\ No newline at end of file
diff --git a/Engine/Plugins/Experimental/SpatialReadiness/Source/PhysicsSpatialReadiness/Public/PhysicsSpatialReadinessSimCallback.h b/Engine/Plugins/Experimental/SpatialReadiness/Source/PhysicsSpatialReadiness/Public/PhysicsSpatialReadinessSimCallback.h
new file mode 100644
--- /dev/null
+++ b/Engine/Plugins/Experimental/SpatialReadiness/Source/PhysicsSpatialReadiness/Public/PhysicsSpatialReadinessSimCallback.h
@@ -0,0 +1,74 @@
+// Copyright Epic Games, Inc. All Rights Reserved.
+
+#pragma once
+
+#include "Chaos/SimCallbackObject.h"
+#include "PhysicsProxy/SingleParticlePhysicsProxy.h"
+#include "Chaos/Framework/HashMappedArray.h"
+
+using namespace Chaos;
+/*
+using Chaos::FSimCallbackInput;
+using Chaos::FSimCallbackOutput;
+using Chaos::FSimCallbackNoInput;
+using Chaos::FSimCallbackNoOutput;
+using Chaos::TSimCallbackObject;
+using Chaos::FSingleParticlePhysicsProxy;
+using Chaos::FGeometryParticleHandle;
+using Chaos::FMidPhaseModifierAccessor;
+*/
+
+struct FPhysicsSpatialReadinessSimCallbackInput : public FSimCallbackInput
+{
+	TSet<FSingleParticlePhysicsProxy*> UnreadyVolumesToAdd;
+	TSet<FSingleParticlePhysicsProxy*> UnreadyVolumesToRemove;
+
+	void Reset()
+	{
+		UnreadyVolumesToAdd.Reset();
+		UnreadyVolumesToRemove.Reset();
+	}
+};
+
+struct FPhysicsSpatialReadinessSimCallback : public TSimCallbackObject<
+	FPhysicsSpatialReadinessSimCallbackInput,
+	FSimCallbackNoOutput,
+	ESimCallbackOptions::Presimulate |
+	ESimCallbackOptions::MidPhaseModification |
+	ESimCallbackOptions::PostIntegrate>
+{
+public:
+	FPhysicsSpatialReadinessSimCallback()
+		: HashSize(1024)
+		, ParticleDataCache(HashSize)
+	{ }
+
+protected:
+
+	/* begin: TSimCallbackObject */
+	virtual void OnPreSimulate_Internal() override;
+	virtual void OnMidPhaseModification_Internal(FMidPhaseModifierAccessor& Accessor) override;
+	virtual void OnPostIntegrate_Internal() override;
+	/* end: TSimCallbackObject */
+
+private:
+
+	// List of particle handles which represent unready volumes
+	TSet<FGeometryParticleHandle*> UnreadyVolumes;
+
+	// Size of the ParticleDataCache
+	uint32 HashSize;
+
+	// This type data is used to create per-particle readiness data.
+	struct FParticleDataCacheTraits
+	{
+		static uint32 GetIDHash(const FGeometryParticleHandle* Particle) { return MurmurFinalize32(uint32(Particle->UniqueIdx().Idx)); }
+		static bool ElementHasID(const FGeometryParticleHandle* A, const FGeometryParticleHandle* B) { return A == B; }
+	};
+	using FParticleDataCache = Private::THashMappedArray<FGeometryParticleHandle*, EObjectStateType, FParticleDataCacheTraits>;
+
+	// "Unready" particles are forced to be stationary in PreSimulate, and restored
+	// to their previous state in PostIntegrate. The values needed for restoration
+	// are stored and mapped back to particle id's in this hash-mapped array.
+	FParticleDataCache ParticleDataCache;
+};
diff --git a/Engine/Plugins/Experimental/SpatialReadiness/Source/PhysicsSpatialReadiness/Public/PhysicsSpatialReadinessSubsystem.h b/Engine/Plugins/Experimental/SpatialReadiness/Source/PhysicsSpatialReadiness/Public/PhysicsSpatialReadinessSubsystem.h
new file mode 100644
--- /dev/null
+++ b/Engine/Plugins/Experimental/SpatialReadiness/Source/PhysicsSpatialReadiness/Public/PhysicsSpatialReadinessSubsystem.h
@@ -0,0 +1,37 @@
+// Copyright Epic Games, Inc. All Rights Reserved.
+
+#pragma once
+
+#include "CoreMinimal.h"
+#include "Subsystems/WorldSubsystem.h"
+#include "SpatialReadinessInterface.h"
+#include "PhysicsSpatialReadinessSimCallback.h"
+#include "PhysicsSpatialReadinessSubsystem.generated.h"
+
+UCLASS()
+class UPhysicsSpatialReadinessSubsystem : public UWorldSubsystem, public ISpatialReadiness
+{
+	GENERATED_BODY()
+
+protected:
+
+	// Begin: UWorldSubsystem
+	virtual bool ShouldCreateSubsystem(UObject* Outer) const override;
+	virtual void OnWorldBeginPlay(UWorld& InWorld) override;
+	virtual void Initialize(FSubsystemCollectionBase& Collection) override;
+	virtual void Deinitialize() override;
+	// End: UWorldSubsystem
+
+	// Begin: ISpatialReadiness
+	virtual FUnreadyVolumeHandle AddUnreadyVolume(const FBox& AABB, const FString& Description) override;
+	virtual void RemoveUnreadyVolume(FUnreadyVolumeHandle Handle) override;
+	// End: ISpatialReadiness
+
+private:
+	bool CreateSimCallback();
+	bool DestroySimCallback();
+
+	FPhysicsSpatialReadinessSimCallback* SimCallback = nullptr;
+
+	TSparseArray<FSingleParticlePhysicsProxy*> UnreadyVolumes;
+};
diff --git a/Engine/Plugins/Experimental/SpatialReadiness/Source/SpatialReadiness/Private/SpatialReadiness/SpatialReadinessModule.cpp b/Engine/Plugins/Experimental/SpatialReadiness/Source/SpatialReadiness/Private/SpatialReadiness/SpatialReadinessModule.cpp
new file mode 100644
--- /dev/null
+++ b/Engine/Plugins/Experimental/SpatialReadiness/Source/SpatialReadiness/Private/SpatialReadiness/SpatialReadinessModule.cpp
@@ -0,0 +1,10 @@
+// Copyright Epic Games, Inc. All Rights Reserved.
+
+#include "SpatialReadinessModule.h"
+#include "Modules/ModuleManager.h"
+
+IMPLEMENT_MODULE(FDefaultModuleImpl, SpatialReadiness)
+
+void FSpatialReadinessModule::StartupModule() { }
+
+void FSpatialReadinessModule::ShutdownModule() { }
\ No newline at end of file
diff --git a/Engine/Plugins/Experimental/SpatialReadiness/Source/SpatialReadiness/Private/SpatialReadiness/UnreadyVolumeHandle.cpp b/Engine/Plugins/Experimental/SpatialReadiness/Source/SpatialReadiness/Private/SpatialReadiness/UnreadyVolumeHandle.cpp
new file mode 100644
--- /dev/null
+++ b/Engine/Plugins/Experimental/SpatialReadiness/Source/SpatialReadiness/Private/SpatialReadiness/UnreadyVolumeHandle.cpp
@@ -0,0 +1,13 @@
+// Copyright Epic Games, Inc. All Rights Reserved.
+
+#include "UnreadyVolumeHandle.h"
+
+void FUnreadyVolumeHandle::Destroy()
+{
+	// todo
+}
+
+void FUnreadyVolumeHandle::Recreate()
+{
+	// todo
+}
diff --git a/Engine/Plugins/Experimental/SpatialReadiness/Source/SpatialReadiness/Public/SpatialReadinessInterface.h b/Engine/Plugins/Experimental/SpatialReadiness/Source/SpatialReadiness/Public/SpatialReadinessInterface.h
new file mode 100644
--- /dev/null
+++ b/Engine/Plugins/Experimental/SpatialReadiness/Source/SpatialReadiness/Public/SpatialReadinessInterface.h
@@ -0,0 +1,17 @@
+// Copyright Epic Games, Inc. All Rights Reserved.
+
+#pragma once
+
+#include "CoreMinimal.h"
+#include "UnreadyVolumeHandle.h"
+
+class FUnreadyVolumeHandle;
+
+class ISpatialReadiness
+{
+public:
+	virtual ~ISpatialReadiness() = default;
+
+	virtual FUnreadyVolumeHandle AddUnreadyVolume(const FBox& AABB, const FString& Description) = 0;
+	virtual void RemoveUnreadyVolume(FUnreadyVolumeHandle Handle) = 0;
+};
\ No newline at end of file
diff --git a/Engine/Plugins/Experimental/SpatialReadiness/Source/SpatialReadiness/Public/SpatialReadinessModule.h b/Engine/Plugins/Experimental/SpatialReadiness/Source/SpatialReadiness/Public/SpatialReadinessModule.h
new file mode 100644
--- /dev/null
+++ b/Engine/Plugins/Experimental/SpatialReadiness/Source/SpatialReadiness/Public/SpatialReadinessModule.h
@@ -0,0 +1,13 @@
+// Copyright Epic Games, Inc. All Rights Reserved.
+
+#pragma once
+
+#include "CoreMinimal.h"
+#include "Modules/ModuleManager.h"
+
+class FSpatialReadinessModule : public IModuleInterface
+{
+public:
+	virtual void StartupModule() override;
+	virtual void ShutdownModule() override;
+};
\ No newline at end of file
diff --git a/Engine/Plugins/Experimental/SpatialReadiness/Source/SpatialReadiness/Public/UnreadyVolumeHandle.h b/Engine/Plugins/Experimental/SpatialReadiness/Source/SpatialReadiness/Public/UnreadyVolumeHandle.h
new file mode 100644
--- /dev/null
+++ b/Engine/Plugins/Experimental/SpatialReadiness/Source/SpatialReadiness/Public/UnreadyVolumeHandle.h
@@ -0,0 +1,13 @@
+// Copyright Epic Games, Inc. All Rights Reserved.
+
+#pragma once
+
+#include "CoreMinimal.h"
+
+class ISpatialReadiness;
+
+class FUnreadyVolumeHandle
+{
+	void Destroy();
+	void Recreate();
+};
diff --git a/Engine/Plugins/Experimental/SpatialReadiness/Source/SpatialReadiness/SpatialReadiness.Build.cs b/Engine/Plugins/Experimental/SpatialReadiness/Source/SpatialReadiness/SpatialReadiness.Build.cs
new file mode 100644
--- /dev/null
+++ b/Engine/Plugins/Experimental/SpatialReadiness/Source/SpatialReadiness/SpatialReadiness.Build.cs
@@ -0,0 +1,21 @@
+// Copyright Epic Games, Inc. All Rights Reserved.
+
+using UnrealBuildTool;
+
+public class SpatialReadiness : ModuleRules
+{
+    public SpatialReadiness(ReadOnlyTargetRules Target) : base(Target)
+    {
+        PrivateDependencyModuleNames.AddRange(
+            new string[]
+            {
+                "Engine",
+            });
+
+        PublicDependencyModuleNames.AddRange(
+            new string[] {
+                "Core",
+            }
+        );
+    }
+}
diff --git a/Engine/Plugins/Experimental/SpatialReadiness/SpatialReadiness.uplugin b/Engine/Plugins/Experimental/SpatialReadiness/SpatialReadiness.uplugin
new file mode 100644
--- /dev/null
+++ b/Engine/Plugins/Experimental/SpatialReadiness/SpatialReadiness.uplugin
@@ -0,0 +1,29 @@
+{
+	"FileVersion": 1,
+	"CanContainContent": false,
+	"ExplicitlyLoaded": false,
+	"EnabledByDefault": false,
+	"Modules": [
+		{
+			"Name": "SpatialReadiness",
+			"Type": "Runtime",
+			"LoadingPhase": "Default"
+		},
+		{
+			"Name": "PhysicsSpatialReadiness",
+			"Type": "Runtime",
+			"LoadingPhase": "Default"
+		},
+		{
+			"Name": "SpatialReadinessTests",
+			"Type": "Runtime",
+			"LoadingPhase": "Default"
+		}
+	],
+	"Plugins": [
+		{
+			"Name": "ChaosUserDataPT",
+			"Enabled": true
+		}
+	]
+}
diff --git a/Engine/Plugins/Experimental/SpatialReadiness/Tests/Private/PhysicsSpatialReadinessTests.cpp b/Engine/Plugins/Experimental/SpatialReadiness/Tests/Private/PhysicsSpatialReadinessTests.cpp
new file mode 100644
--- /dev/null
+++ b/Engine/Plugins/Experimental/SpatialReadiness/Tests/Private/PhysicsSpatialReadinessTests.cpp
@@ -0,0 +1,25 @@
+// Copyright Epic Games, Inc. All Rights Reserved.
+#include "TestHarness.h"
+#include "CoreMinimal.h"
+//#include "Chaos/Framework/ChaosUserDataPT.h"
+#include "PBDRigidsSolver.h"
+#include "ChaosSolversModule.h"
+#include "PhysicsSpatialReadinessSimCallback.h"
+
+
+TEST_CASE("Sim Callback", "[physics spatial readiness]")
+{
+	const float DeltaTime = 1.f/30.f;
+
+	// Create a solver in the solvers module
+	FChaosSolversModule* Module = FChaosSolversModule::GetModule();
+	Chaos::FPBDRigidsSolver* Solver = Module->CreateSolver(nullptr, /*AsyncDt=*/DeltaTime, Chaos::EThreadingMode::TaskGraph);
+
+	// Create a test PhysicsSpatialReadinessSimCallback
+	FPhysicsSpatialReadinessSimCallback* SimCallback = Solver->CreateAndRegisterSimCallbackObject_External<FPhysicsSpatialReadinessSimCallback>();
+
+	SECTION("Create unready volume")
+	{
+		// TODO: Everything
+	}
+}
\ No newline at end of file
diff --git a/Engine/Plugins/Experimental/SpatialReadiness/Tests/Private/SpatialReadinessTests.cpp b/Engine/Plugins/Experimental/SpatialReadiness/Tests/Private/SpatialReadinessTests.cpp
new file mode 100644
--- /dev/null
+++ b/Engine/Plugins/Experimental/SpatialReadiness/Tests/Private/SpatialReadinessTests.cpp
@@ -0,0 +1,11 @@
+// Copyright Epic Games, Inc. All Rights Reserved.
+#include "TestHarness.h"
+
+TEST_CASE("...", "[spatial readiness]")
+{
+	SECTION("...")
+	{
+		// TODO: everything
+		REQUIRE(true);
+	}
+}
\ No newline at end of file
diff --git a/Engine/Plugins/Experimental/SpatialReadiness/Tests/Private/TestGroupEvents.cpp b/Engine/Plugins/Experimental/SpatialReadiness/Tests/Private/TestGroupEvents.cpp
new file mode 100644
--- /dev/null
+++ b/Engine/Plugins/Experimental/SpatialReadiness/Tests/Private/TestGroupEvents.cpp
@@ -0,0 +1,15 @@
+// Copyright Epic Games, Inc. All Rights Reserved.
+
+#include "TestCommon/Initialization.h"
+
+#include <catch2/catch_test_macros.hpp>
+
+GROUP_BEFORE_GLOBAL(Catch::DefaultGroup)
+{
+	InitAll(/*bAllowLogging=*/true, /*bMultithreaded=*/false);
+}
+
+GROUP_AFTER_GLOBAL(Catch::DefaultGroup)
+{
+	CleanupAll();
+}
diff --git a/Engine/Plugins/Experimental/SpatialReadiness/Tests/SpatialReadinessTests.Build.cs b/Engine/Plugins/Experimental/SpatialReadiness/Tests/SpatialReadinessTests.Build.cs
new file mode 100644
--- /dev/null
+++ b/Engine/Plugins/Experimental/SpatialReadiness/Tests/SpatialReadinessTests.Build.cs
@@ -0,0 +1,29 @@
+// Copyright Epic Games, Inc. All Rights Reserved.
+
+using UnrealBuildTool;
+
+public class SpatialReadinessTests : TestModuleRules
+{
+	static SpatialReadinessTests()
+	{
+		TestMetadata = new Metadata()
+		{
+			TestName = "SpatialReadinessTests",
+			TestShortName = "SpatialReadinessTests",
+		};
+	}
+
+	public SpatialReadinessTests(ReadOnlyTargetRules Target) : base(Target)
+	{
+		PrivateDependencyModuleNames.AddRange(new string[]
+		{
+			"ApplicationCore",
+			"Core",
+			//"CoreUObject",
+			"SpatialReadiness",
+			"PhysicsSpatialReadiness",
+			"Chaos",
+			//"ChaosUserDataPT",
+		});
+	}
+}
\ No newline at end of file
diff --git a/Engine/Plugins/Experimental/SpatialReadiness/Tests/SpatialReadinessTests.Target.cs b/Engine/Plugins/Experimental/SpatialReadiness/Tests/SpatialReadinessTests.Target.cs
new file mode 100644
--- /dev/null
+++ b/Engine/Plugins/Experimental/SpatialReadiness/Tests/SpatialReadinessTests.Target.cs
@@ -0,0 +1,13 @@
+// Copyright Epic Games, Inc. All Rights Reserved.
+
+using UnrealBuildTool;
+
+[SupportedPlatforms(UnrealPlatformClass.All)]
+public class SpatialReadinessTestsTarget : TestTargetRules
+{
+	public SpatialReadinessTestsTarget(TargetInfo Target) : base(Target)
+	{
+		bUsePlatformFileStub = true;
+		bMockEngineDefaults = true;
+	}
+}
