CIS - Fix linker error when building the UE editor without a project.

#rnx

diff --git a/Engine/Plugins/Experimental/SpatialReadiness/Source/PhysicsSpatialReadiness/PhysicsSpatialReadiness.Build.cs b/Engine/Plugins/Experimental/SpatialReadiness/Source/PhysicsSpatialReadiness/PhysicsSpatialReadiness.Build.cs
--- a/Engine/Plugins/Experimental/SpatialReadiness/Source/PhysicsSpatialReadiness/PhysicsSpatialReadiness.Build.cs
+++ b/Engine/Plugins/Experimental/SpatialReadiness/Source/PhysicsSpatialReadiness/PhysicsSpatialReadiness.Build.cs
@@ -9,6 +9,7 @@ public class PhysicsSpatialReadiness : ModuleRules
         PrivateDependencyModuleNames.AddRange(
             new string[]
             {
+                "CoreUObject",
                 "Engine",
                 "SpatialReadiness",
             });
