Remove SpatialReadinessTests module from direct reference in the SpatialReadiness plugin.

diff --git a/Engine/Plugins/Experimental/SpatialReadiness/SpatialReadiness.uplugin b/Engine/Plugins/Experimental/SpatialReadiness/SpatialReadiness.uplugin
--- a/Engine/Plugins/Experimental/SpatialReadiness/SpatialReadiness.uplugin
+++ b/Engine/Plugins/Experimental/SpatialReadiness/SpatialReadiness.uplugin
@@ -13,11 +13,6 @@
 			"Name": "PhysicsSpatialReadiness",
 			"Type": "Runtime",
 			"LoadingPhase": "Default"
-		},
-		{
-			"Name": "SpatialReadinessTests",
-			"Type": "Runtime",
-			"LoadingPhase": "Default"
 		}
 	],
 	"Plugins": [
