Don't try to remove simcallback from spatial readiness subsystem if it doesn't exist

diff --git a/Engine/Plugins/Experimental/SpatialReadiness/Source/PhysicsSpatialReadiness/Private/PhysicsSpatialReadiness/PhysicsSpatialReadinessSubsystem.cpp b/Engine/Plugins/Experimental/SpatialReadiness/Source/PhysicsSpatialReadiness/Private/PhysicsSpatialReadiness/PhysicsSpatialReadinessSubsystem.cpp
--- a/Engine/Plugins/Experimental/SpatialReadiness/Source/PhysicsSpatialReadiness/Private/PhysicsSpatialReadiness/PhysicsSpatialReadinessSubsystem.cpp
+++ b/Engine/Plugins/Experimental/SpatialReadiness/Source/PhysicsSpatialReadiness/Private/PhysicsSpatialReadiness/PhysicsSpatialReadinessSubsystem.cpp
@@ -66,11 +66,14 @@ bool UPhysicsSpatialReadinessSubsystem::CreateSimCallback()
 
 bool UPhysicsSpatialReadinessSubsystem::DestroySimCallback()
 {
-	if (Chaos::FPhysicsSolver* Solver = GetSolver())
+	if (SimCallback)
 	{
-		Solver->UnregisterAndFreeSimCallbackObject_External(SimCallback);
-		SimCallback = nullptr;
-		return true;
+		if (Chaos::FPhysicsSolver* Solver = GetSolver())
+		{
+			Solver->UnregisterAndFreeSimCallbackObject_External(SimCallback);
+			SimCallback = nullptr;
+			return true;
+		}
 	}
 
 	return false;
