Add implementations for unready volume particle creation, destruction, and tracking

diff --git a/Engine/Plugins/Experimental/SpatialReadiness/Source/PhysicsSpatialReadiness/Private/PhysicsSpatialReadiness/PhysicsSpatialReadinessSimCallback.cpp b/Engine/Plugins/Experimental/SpatialReadiness/Source/PhysicsSpatialReadiness/Private/PhysicsSpatialReadiness/PhysicsSpatialReadinessSimCallback.cpp
--- a/Engine/Plugins/Experimental/SpatialReadiness/Source/PhysicsSpatialReadiness/Private/PhysicsSpatialReadiness/PhysicsSpatialReadinessSimCallback.cpp
+++ b/Engine/Plugins/Experimental/SpatialReadiness/Source/PhysicsSpatialReadiness/Private/PhysicsSpatialReadiness/PhysicsSpatialReadinessSimCallback.cpp
@@ -3,39 +3,254 @@
 #include "PhysicsSpatialReadinessSimCallback.h"
 #include "PhysicsProxy/SingleParticlePhysicsProxy.h"
 #include "Chaos/ChaosEngineInterface.h"
+#include "Chaos/ChaosScene.h"
+#include "Chaos/MidPhaseModification.h"
+#include "Chaos/ParticleHandle.h"
 #include "SpatialReadinessHandle.h"
+#include "PBDRigidsSolver.h"
+#include "Chaos/PBDRigidsEvolutionGBF.h"
+#include "Physics/PhysicsFiltering.h"
 
 using namespace Chaos;
 
-void FPhysicsSpatialReadinessSimCallback::OnPreSimulate_Internal()
+FSpatialReadinessHandle FPhysicsSpatialReadinessSimCallback::AddReadinessVolume(const FBox& Bounds, const FString& Description)
 {
+	return SpatialReadiness.CreateHandle(Bounds, Description);
 }
 
-void FPhysicsSpatialReadinessSimCallback::OnMidPhaseModification_Internal(FMidPhaseModifierAccessor& Accessor)
+FPBDRigidsEvolution* FPhysicsSpatialReadinessSimCallback::GetEvolution()
 {
+	if (FPBDRigidsSolver* MySolver = static_cast<FPBDRigidsSolver*>(GetSolver()))
+	{
+		return MySolver->GetEvolution();
+	}
+
+	return nullptr;
 }
 
-void FPhysicsSpatialReadinessSimCallback::OnPostIntegrate_Internal()
+int32 FPhysicsSpatialReadinessSimCallback::AddUnreadyVolume_GT(const FBox& Bounds, const FString& Description)
+{
+	// Get a pointer to the input struct - if this is null just cancel. If we create a
+	// particle without the input, then we'll lose track of it.
+	FPhysicsSpatialReadinessSimCallbackInput* SimInputPtr = GetProducerInputData_External();
+	if (!ensureMsgf(SimInputPtr, TEXT("Failed to access sim callback object input")))
+	{
+		return INDEX_NONE;
+	}
+	FPhysicsSpatialReadinessSimCallbackInput& SimInput = *SimInputPtr;
+
+	// Create a box implicit geometry from the bounds
+	const FVec3 BoxCenter = (Bounds.Min + Bounds.Max) * .5f;
+	const FVec3 BoxHalfExtent = Bounds.Max - Bounds.Min;
+	Chaos::FImplicitObjectPtr BoxGeom = MakeImplicitObjectPtr<TBox<FReal, 3>>(-BoxHalfExtent, BoxHalfExtent);
+
+	// Create a new static particle to represent the volume
+	FActorCreationParams Params;
+	Params.bSimulatePhysics = false;
+	Params.bStatic = true;
+	Params.InitialTM = FTransform(FQuat::Identity, BoxCenter);
+	Params.Scene = &PhysicsScene;
+	FSingleParticlePhysicsProxy* ParticleProxy = nullptr;
+	FChaosEngineInterface::CreateActor(Params, ParticleProxy);
+	if (!ensureMsgf(ParticleProxy, TEXT("Failed to create new particle proxy")))
+	{
+		return INDEX_NONE;
+	}
+	FRigidBodyHandle_External& ParticleHandle = ParticleProxy->GetGameThreadAPI();
+
+	// Create collision response container
+	FCollisionResponseContainer CollisionResponse;
+	CollisionResponse.SetAllChannels(ECollisionResponse::ECR_Block);
+
+	// Create collision filter data for the particle
+	FCollisionFilterData QueryFilterData, SimFilterData;
+	CreateShapeFilterData(
+		/* MyChannel */            static_cast<uint8>(ECollisionChannel::ECC_WorldStatic),
+		/* MaskFilter */           FMaskFilter(0),
+		/* SourceObjectID */       0,
+		/* ResponseToChannels */   CollisionResponse,
+		/* ComponentID */          0,
+		/* BodyIndex */            0,
+		/* OutQueryData */         QueryFilterData,
+		/* OutSimData */           SimFilterData,
+		/* bEnableCCD */           true,
+		/* bEnableContactNotify */ false,
+		/* bStaticShape */         true);
+
+	// Make the geometry
+	ParticleHandle.SetGeometry(BoxGeom);
+	ParticleHandle.SetShapeSimCollisionEnabled(0, true);
+	ParticleHandle.SetShapeQueryCollisionEnabled(0, false);
+	ParticleHandle.SetShapeSimData(0, SimFilterData);
+	ParticleHandle.ShapesArray()[0]->SetIsProbe(true);
+#if CHAOS_DEBUG_NAME
+	ParticleHandle.SetDebugName(MakeShared<FString>(TEXT("Unready Volume")));
+#endif
+
+	// Add the new particle to the scene
+	TArray<FPhysicsActorHandle> Actors = { ParticleProxy };
+	PhysicsScene.AddActorsToScene_AssumesLocked(Actors);
+
+	// Save the proxy in our list of GT particles
+	const int32 ParticleIndex = UnreadyVolumeParticles_GT.Add(ParticleProxy);
+
+	// Queue up the particle proxy for processing on PT
+	SimInput.UnreadyVolumesToAdd.Add(ParticleProxy);
+	SimInput.UnreadyVolumesToRemove.Remove(ParticleProxy);
+
+	return ParticleIndex;
+}
+
+void FPhysicsSpatialReadinessSimCallback::RemoveUnreadyVolume_GT(int32 UnreadyVolumeIndex)
 {
+	if (!ensureMsgf(UnreadyVolumeParticles_GT.IsValidIndex(UnreadyVolumeIndex), TEXT("Trying to remove unready volume whose index is not being tracked")))
+	{
+		return;
+	}
+
+	// Get the proxy associated with this index
+	FSingleParticlePhysicsProxy* ParticleProxy = UnreadyVolumeParticles_GT[UnreadyVolumeIndex];
+	if (!ensureMsgf(ParticleProxy, TEXT("Particle proxy associated with unready volume index was null")))
+	{
+		return;
+	}
+	
+	// Free the index in our GT tracker
+	UnreadyVolumeParticles_GT.RemoveAt(UnreadyVolumeIndex);
+
+	// Tell the PT to remove it's tracking of this proxy as well
+	FPhysicsSpatialReadinessSimCallbackInput* SimInput = GetProducerInputData_External();
+	if (ensureMsgf(SimInput, TEXT("Failed to access sim input data")))
+	{
+		SimInput->UnreadyVolumesToRemove.Add(ParticleProxy);
+		SimInput->UnreadyVolumesToAdd.Remove(ParticleProxy);
+	}
+
+	// Delete the particle
+	FChaosEngineInterface::ReleaseActor(ParticleProxy, &PhysicsScene);
 }
 
-FSpatialReadinessHandle FPhysicsSpatialReadinessSimCallback::AddReadinessVolume(const FBox& Bounds, const FString& Description)
+void FPhysicsSpatialReadinessSimCallback::FreezeParticles_PT()
 {
-	return SpatialReadiness.CreateHandle(Bounds, Description);
+	FPBDRigidsEvolution* Evolution = GetEvolution();
+	if (!ensureMsgf(Evolution, TEXT("Attempted to freeze particle, but had no evolution")))
+	{
+		return;
+	}
+
+	for (FPBDRigidParticleHandle* RigidParticle : UnreadyRigidParticles_PT)
+	{
+		const EObjectStateType ObjectState = RigidParticle->ObjectState();
+		if (ObjectState == EObjectStateType::Static)
+		{
+			return;
+		}
+
+		// Get the geometry particle
+		FGeometryParticleHandle* GeometryParticle = RigidParticle;
+
+		// Cache the current object state
+		ParticleDataCache_PT.Add({ GeometryParticle, ObjectState });
+
+		// Set the object state to static
+		Evolution->SetParticleObjectState(RigidParticle, EObjectStateType::Static);
+	}
+}
+
+void FPhysicsSpatialReadinessSimCallback::UnFreezeParticles_PT()
+{
+	FPBDRigidsEvolution* Evolution = GetEvolution();
+	if (!ensureMsgf(Evolution, TEXT("Attempted to un-freeze particles, but had no evolution")))
+	{
+		return;
+	}
+
+	// For each particle that we froze, unfreeze it
+	for (TPair<FGeometryParticleHandle*, EObjectStateType>& ParticleData : ParticleDataCache_PT)
+	{
+		FPBDRigidParticleHandle* RigidParticle = ParticleData.Key->CastToRigidParticle();
+		if (RigidParticle == nullptr)
+		{
+			continue;
+		}
+
+		// Restore the object state of this particle
+		Evolution->SetParticleObjectState(RigidParticle, ParticleData.Value);
+	}
+
+	// Clear the array
+	ParticleDataCache_PT.Reset();
+}
+
+void FPhysicsSpatialReadinessSimCallback::OnPreSimulate_Internal()
+{
+	// Process inputs we may have gotten from the game thread
+	if (const FPhysicsSpatialReadinessSimCallbackInput* Input = GetConsumerInput_Internal())
+	{
+		// Process additions
+		for (FSingleParticlePhysicsProxy* ParticleProxy : Input->UnreadyVolumesToAdd)
+		{
+			UnreadyVolumeParticles_PT.Add(ParticleProxy);
+		}
+
+		// Process removals
+		for (FSingleParticlePhysicsProxy* ParticleProxy : Input->UnreadyVolumesToRemove)
+		{
+			UnreadyVolumeParticles_PT.Remove(ParticleProxy);
+		}
+	}
 }
 
-int32 FPhysicsSpatialReadinessSimCallback::AddUnreadyVolume_Internal(const FBox& Bounds, const FString& Description)
+void FPhysicsSpatialReadinessSimCallback::OnMidPhaseModification_Internal(FMidPhaseModifierAccessor& Accessor)
 {
-	//
-	// TODO
-	//
+	// Clear the list of unready particles
+	UnreadyRigidParticles_PT.Reset();
+
+	// Go through every unready volume
+	for (FSingleParticlePhysicsProxy* UnreadyProxy : UnreadyVolumeParticles_PT)
+	{
+		if (UnreadyProxy->GetMarkedDeleted())
+		{
+			continue;
+		}
+
+		FGeometryParticleHandle* UnreadyVolume = UnreadyProxy->GetHandle_LowLevel();
+		if (UnreadyVolume == nullptr)
+		{
+			continue;
+		}
+
+		// Go through every mid-phase which involves this volume
+		for (FMidPhaseModifier& MidPhase : Accessor.GetMidPhases(UnreadyVolume))
+		{
+			// Get the particle that is not the unready volume
+			FGeometryParticleHandle* GeometryParticle = MidPhase.GetOtherParticle(UnreadyVolume);
+			if (GeometryParticle == nullptr)
+			{
+				continue;
+			}
 
-	return INDEX_NONE;
+			FPBDRigidParticleHandle* RigidParticle = GeometryParticle->CastToRigidParticle();
+			if (RigidParticle == nullptr)
+			{
+				continue;
+			}
+
+			// Add the particle to the list of particles to freeze
+			UnreadyRigidParticles_PT.Add(RigidParticle);
+
+			// Disable the midphase
+			MidPhase.Disable();
+		}
+	}
+}
+
+void FPhysicsSpatialReadinessSimCallback::OnPreIntegrate_Internal()
+{
+	FreezeParticles_PT();
 }
 
-void FPhysicsSpatialReadinessSimCallback::RemoveUnreadyVolume_Internal(int32 UnreadyVolumeIndex)
+void FPhysicsSpatialReadinessSimCallback::OnPostIntegrate_Internal()
 {
-	//
-	// TODO
-	//
+	UnFreezeParticles_PT();
 }
diff --git a/Engine/Plugins/Experimental/SpatialReadiness/Source/PhysicsSpatialReadiness/Private/PhysicsSpatialReadiness/PhysicsSpatialReadinessSubsystem.cpp b/Engine/Plugins/Experimental/SpatialReadiness/Source/PhysicsSpatialReadiness/Private/PhysicsSpatialReadiness/PhysicsSpatialReadinessSubsystem.cpp
--- a/Engine/Plugins/Experimental/SpatialReadiness/Source/PhysicsSpatialReadiness/Private/PhysicsSpatialReadiness/PhysicsSpatialReadinessSubsystem.cpp
+++ b/Engine/Plugins/Experimental/SpatialReadiness/Source/PhysicsSpatialReadiness/Private/PhysicsSpatialReadiness/PhysicsSpatialReadinessSubsystem.cpp
@@ -55,10 +55,16 @@ bool UPhysicsSpatialReadinessSubsystem::CreateSimCallback()
 		return false;
 	}
 
-	if (Chaos::FPhysicsSolver* Solver = GetSolver())
+	if (FPhysScene_Chaos* Scene = GetScene())
 	{
-		SimCallback = Solver->CreateAndRegisterSimCallbackObject_External<FPhysicsSpatialReadinessSimCallback>();
-		return true;
+		if (Chaos::FPhysicsSolver* Solver = GetSolver())
+		{
+			SimCallback = Solver->CreateAndRegisterSimCallbackObject_External<FPhysicsSpatialReadinessSimCallback>(static_cast<FChaosScene&>(*Scene));
+			if (ensureMsgf(SimCallback, TEXT("Sim callback creation failed")))
+			{
+				return true;
+			}
+		}
 	}
 
 	return false;
@@ -79,16 +85,25 @@ bool UPhysicsSpatialReadinessSubsystem::DestroySimCallback()
 	return false;
 }
 
-Chaos::FPhysicsSolver* UPhysicsSpatialReadinessSubsystem::GetSolver()
+FPhysScene_Chaos* UPhysicsSpatialReadinessSubsystem::GetScene()
 {
 	if (UWorld* World = GetWorld())
 	{
 		if (FPhysScene* Scene = World->GetPhysicsScene())
 		{
-			if (Chaos::FPhysicsSolver* Solver = Scene->GetSolver())
-			{
-				return Solver;
-			}
+			return static_cast<FPhysScene_Chaos*>(Scene);
+		}
+	}
+	return nullptr;
+}
+
+Chaos::FPhysicsSolver* UPhysicsSpatialReadinessSubsystem::GetSolver()
+{
+	if (FPhysScene* Scene = GetScene())
+	{
+		if (Chaos::FPhysicsSolver* Solver = Scene->GetSolver())
+		{
+			return Solver;
 		}
 	}
 
diff --git a/Engine/Plugins/Experimental/SpatialReadiness/Source/PhysicsSpatialReadiness/Public/PhysicsSpatialReadinessSimCallback.h b/Engine/Plugins/Experimental/SpatialReadiness/Source/PhysicsSpatialReadiness/Public/PhysicsSpatialReadinessSimCallback.h
--- a/Engine/Plugins/Experimental/SpatialReadiness/Source/PhysicsSpatialReadiness/Public/PhysicsSpatialReadinessSimCallback.h
+++ b/Engine/Plugins/Experimental/SpatialReadiness/Source/PhysicsSpatialReadiness/Public/PhysicsSpatialReadinessSimCallback.h
@@ -3,12 +3,19 @@
 #pragma once
 
 #include "Chaos/SimCallbackObject.h"
-#include "PhysicsProxy/SingleParticlePhysicsProxy.h"
-#include "Chaos/Framework/HashMappedArray.h"
+#include "Chaos/PBDRigidsEvolutionFwd.h"
+#include "Chaos/Particle/ObjectState.h"
+#include "PhysicsSolver.h"
 #include "SpatialReadinessAPI.h"
 
 using namespace Chaos;
 
+class FChaosScene;
+namespace Chaos
+{
+	class FSingleParticlePhysicsProxy;
+}
+
 struct FPhysicsSpatialReadinessSimCallbackInput : public FSimCallbackInput
 {
 	TSet<FSingleParticlePhysicsProxy*> UnreadyVolumesToAdd;
@@ -27,32 +34,44 @@ struct FPhysicsSpatialReadinessSimCallback
 		FSimCallbackNoOutput,
 		ESimCallbackOptions::Presimulate |
 		ESimCallbackOptions::MidPhaseModification |
+		ESimCallbackOptions::PreIntegrate |
 		ESimCallbackOptions::PostIntegrate>
 {
 	using This = FPhysicsSpatialReadinessSimCallback;
 
 public:
-	FPhysicsSpatialReadinessSimCallback()
-		: SpatialReadiness(this, &This::AddUnreadyVolume_Internal, &This::RemoveUnreadyVolume_Internal)
-		, HashSize(1024)
-		, ParticleDataCache_PT(HashSize)
+	FPhysicsSpatialReadinessSimCallback(FChaosScene& InPhysicsScene)
+		: PhysicsScene(InPhysicsScene)
+		, SpatialReadiness(this, &This::AddUnreadyVolume_GT, &This::RemoveUnreadyVolume_GT)
+		, ParticleDataCache_PT()
 	{ }
 
+	FSpatialReadinessHandle AddReadinessVolume(const FBox& Bounds, const FString& Description);
+
 	/* begin: TSimCallbackObject */
 protected:
 	virtual void OnPreSimulate_Internal() override;
 	virtual void OnMidPhaseModification_Internal(FMidPhaseModifierAccessor& Accessor) override;
+	virtual void OnPreIntegrate_Internal() override;
 	virtual void OnPostIntegrate_Internal() override;
 	/* end: TSimCallbackObject */
 
-public:
-	FSpatialReadinessHandle AddReadinessVolume(const FBox& Bounds, const FString& Description);
+	// Helpers
+	FPBDRigidsEvolution* GetEvolution();
+
+	// The functions to bind to the spatial readiness API.
+	int32 AddUnreadyVolume_GT(const FBox& Bounds, const FString& Description);
+	void RemoveUnreadyVolume_GT(int32 UnreadyVolumeIndex);
+
+	// Functions for freezing and unfreezing all particles in the
+	// UnreadyRigidParticles_PT set.
+	void FreezeParticles_PT();
+	void UnFreezeParticles_PT();
 
-private:
+	FChaosScene& PhysicsScene;
 
+	// The spatial readiness API object which
 	FSpatialReadinessAPI SpatialReadiness;
-	int32 AddUnreadyVolume_Internal(const FBox& Bounds, const FString& Description);
-	void RemoveUnreadyVolume_Internal(int32 UnreadyVolumeIndex);
 
 	// List of unready volume physics proxies. We directly use single particle
 	// physics proxy rather than something more generic because we know that
@@ -60,28 +79,13 @@ private:
 	TSparseArray<FSingleParticlePhysicsProxy*> UnreadyVolumeParticles_GT;
 
 	// List of particle handles which represent unready volumes
-	TSet<FGeometryParticleHandle*> UnreadyVolumeParticles_PT;
+	TSet<FSingleParticlePhysicsProxy*> UnreadyVolumeParticles_PT;
 
-	// Size of the ParticleDataCache
-	const uint32 HashSize;
-
-	// This type data is used to create per-particle readiness data.
-	struct FParticleDataCacheTraits
-	{
-		static uint32 GetIDHash(const FGeometryParticleHandle* Particle)
-		{
-			return MurmurFinalize32(uint32(Particle->UniqueIdx().Idx));
-		}
-		static bool ElementHasID(const FGeometryParticleHandle* A, const FGeometryParticleHandle* B)
-		{
-			return A == B;
-		}
-	};
-
-	using FParticleDataCache = Private::THashMappedArray<FGeometryParticleHandle*, EObjectStateType, FParticleDataCacheTraits>;
+	// List of particle handles which represent particles that interacted with unready volumes
+	TSet<FPBDRigidParticleHandle*> UnreadyRigidParticles_PT;
 
 	// "Unready" particles are forced to be stationary in PreSimulate, and restored
 	// to their previous state in PostIntegrate. The values needed for restoration
 	// are stored and mapped back to particle id's in this hash-mapped array.
-	FParticleDataCache ParticleDataCache_PT;
+	TArray<TPair<FGeometryParticleHandle*, EObjectStateType>> ParticleDataCache_PT;
 };
diff --git a/Engine/Plugins/Experimental/SpatialReadiness/Source/PhysicsSpatialReadiness/Public/PhysicsSpatialReadinessSubsystem.h b/Engine/Plugins/Experimental/SpatialReadiness/Source/PhysicsSpatialReadiness/Public/PhysicsSpatialReadinessSubsystem.h
--- a/Engine/Plugins/Experimental/SpatialReadiness/Source/PhysicsSpatialReadiness/Public/PhysicsSpatialReadinessSubsystem.h
+++ b/Engine/Plugins/Experimental/SpatialReadiness/Source/PhysicsSpatialReadiness/Public/PhysicsSpatialReadinessSubsystem.h
@@ -27,6 +27,7 @@ public:
 private:
 	bool CreateSimCallback();
 	bool DestroySimCallback();
+	FPhysScene_Chaos* GetScene();
 	Chaos::FPhysicsSolver* GetSolver();
 
 	FPhysicsSpatialReadinessSimCallback* SimCallback = nullptr;
diff --git a/Engine/Plugins/Experimental/SpatialReadiness/Tests/Private/PhysicsSpatialReadinessTests.cpp b/Engine/Plugins/Experimental/SpatialReadiness/Tests/Private/PhysicsSpatialReadinessTests.cpp
--- a/Engine/Plugins/Experimental/SpatialReadiness/Tests/Private/PhysicsSpatialReadinessTests.cpp
+++ b/Engine/Plugins/Experimental/SpatialReadiness/Tests/Private/PhysicsSpatialReadinessTests.cpp
@@ -3,10 +3,16 @@
 #include "CoreMinimal.h"
 #include "PBDRigidsSolver.h"
 #include "ChaosSolversModule.h"
+#include "Chaos/ChaosScene.h"
 #include "PhysicsSpatialReadinessSimCallback.h"
+#include "Engine/EngineTypes.h"
+#include "PhysicsProxy/SingleParticlePhysicsProxy.h"
+#include "Physics/PhysicsFiltering.h"
 
 namespace
 {
+	const FBox UnitBounds(FVector(-.5), FVector(.5));
+
 	void AdvanceAndWait(Chaos::FPBDRigidsSolver* InSolver, float DeltaTime = 1)
 	{
 		if (InSolver)
@@ -15,24 +21,108 @@ namespace
 			InSolver->WaitOnPendingTasks_External();
 		}
 	}
+
+	FSingleParticlePhysicsProxy* MakeDynamicBox(FChaosScene& Scene, const FBox& Bounds, const bool bGravityEnabled=true)
+	{
+		// Create a box implicit geometry from the same bounds as the unready volume
+		const FVec3 BoxCenter = (Bounds.Min + Bounds.Max) * .5f;
+		const FVec3 BoxHalfExtent = Bounds.Max - Bounds.Min;
+		Chaos::FImplicitObjectPtr BoxGeom = MakeImplicitObjectPtr<TBox<FReal, 3>>(-BoxHalfExtent, BoxHalfExtent);
+
+		// Create a new static particle to represent the volume
+		FActorCreationParams Params;
+		Params.bSimulatePhysics = true;
+		Params.bStatic = false;
+		Params.InitialTM = FTransform(FQuat::Identity, BoxCenter);
+		Params.Scene = &Scene;
+		FSingleParticlePhysicsProxy* ParticleProxy = nullptr;
+		FChaosEngineInterface::CreateActor(Params, ParticleProxy);
+		FRigidBodyHandle_External& ParticleHandle = ParticleProxy->GetGameThreadAPI();
+		REQUIRE(ParticleProxy);
+
+		// Create collision response container
+		FCollisionResponseContainer CollisionResponse;
+		CollisionResponse.SetAllChannels(ECollisionResponse::ECR_Block);
+
+		// Create collision filter data for the particle
+		FCollisionFilterData QueryFilterData, SimFilterData;
+		CreateShapeFilterData(
+			/* MyChannel */            static_cast<uint8>(ECollisionChannel::ECC_WorldDynamic),
+			/* MaskFilter */           FMaskFilter(0),
+			/* SourceObjectID */       0,
+			/* ResponseToChannels */   CollisionResponse,
+			/* ComponentID */          0,
+			/* BodyIndex */            0,
+			/* OutQueryData */         QueryFilterData,
+			/* OutSimData */           SimFilterData,
+			/* bEnableCCD */           true,
+			/* bEnableContactNotify */ false,
+			/* bStaticShape */         false);
+
+		// Make the geometry
+		ParticleHandle.SetGeometry(BoxGeom);
+		ParticleHandle.SetShapeSimCollisionEnabled(0, true);
+		ParticleHandle.SetShapeQueryCollisionEnabled(0, false);
+		ParticleHandle.SetShapeSimData(0, SimFilterData);
+		ParticleHandle.SetGravityEnabled(bGravityEnabled);
+		ParticleHandle.SetObjectState(EObjectStateType::Dynamic, false, true);
+#if CHAOS_DEBUG_NAME
+		ParticleHandle.SetDebugName(MakeShared<FString>(TEXT("Dynamic Box")));
+#endif
+
+		// Add the new particle to the scene
+		TArray<FPhysicsActorHandle> Actors = { ParticleProxy };
+		Scene.AddActorsToScene_AssumesLocked(Actors);
+
+		// Return the particle proxy
+		return ParticleProxy;
+	}
+
+	class FTestPhysicsSpatialReadinessSimCallback : public FPhysicsSpatialReadinessSimCallback
+	{
+	public:
+		FTestPhysicsSpatialReadinessSimCallback(FChaosScene& InScene) : FPhysicsSpatialReadinessSimCallback(InScene) { }
+
+		const TSparseArray<FSingleParticlePhysicsProxy*>& GetUnreadyVolumeParticles_GT() const
+		{
+			return UnreadyVolumeParticles_GT;
+		}
+
+		const TSet<FSingleParticlePhysicsProxy*>& GetUnreadyVolumeParticles_PT() const
+		{
+			return UnreadyVolumeParticles_PT;
+		}
+
+		const TSet<FPBDRigidParticleHandle*>& GetUnreadyRigidParticles_PT() const
+		{
+			return UnreadyRigidParticles_PT;
+		}
+	};
 }
 
 TEST_CASE("Sim Callback", "[physics spatial readiness]")
 {
 	// Create a solver in the solvers module
-	FChaosSolversModule* Module = FChaosSolversModule::GetModule();
-	Chaos::FPBDRigidsSolver* Solver = Module->CreateSolver(nullptr, /*AsyncDt=*/1, Chaos::EThreadingMode::TaskGraph);
+	FChaosScene Scene(nullptr, /*AsyncDt=*/-1);
+	Chaos::FPBDRigidsSolver* Solver = Scene.GetSolver();
+	Solver->SetThreadingMode_External(EThreadingModeTemp::TaskGraph);
+
+	// Create a test PhysicsSpatialReadinessSimCallback
+	FTestPhysicsSpatialReadinessSimCallback* SimCallback = Solver->CreateAndRegisterSimCallbackObject_External<FTestPhysicsSpatialReadinessSimCallback>(Scene);
+	AdvanceAndWait(Solver);
 
 	SECTION("Create unready volume, and let the readiness system go out of scope")
 	{
-		// Create a test PhysicsSpatialReadinessSimCallback
-		FPhysicsSpatialReadinessSimCallback* SimCallback = Solver->CreateAndRegisterSimCallbackObject_External<FPhysicsSpatialReadinessSimCallback>();
-		AdvanceAndWait(Solver);
-
 		// Make a volume
-		const FBox Box(FVector(-1), FVector(1));
-		FSpatialReadinessHandle Handle = SimCallback->AddReadinessVolume(Box, TEXT("Test Volume"));
+		FSpatialReadinessHandle Handle = SimCallback->AddReadinessVolume(UnitBounds, TEXT("Test Volume"));
+		REQUIRE(SimCallback->GetUnreadyVolumeParticles_GT().Num() == 1);
+		REQUIRE(SimCallback->GetUnreadyVolumeParticles_PT().Num() == 0);
+
+		// Advance the solver, which should result in the unready volume existing
+		// on the physics thread
 		AdvanceAndWait(Solver);
+		REQUIRE(SimCallback->GetUnreadyVolumeParticles_GT().Num() == 1);
+		REQUIRE(SimCallback->GetUnreadyVolumeParticles_PT().Num() == 1);
 
 		// The volume handle should be "valid"
 		REQUIRE(Handle.IsValid());
@@ -44,4 +134,82 @@ TEST_CASE("Sim Callback", "[physics spatial readiness]")
 		// The volume handle should now be "invalid"
 		REQUIRE(!Handle.IsValid());
 	}
-}
\ No newline at end of file
+
+	SECTION("Spawn dynamic particle overlapping unready volume, no gravity")
+	{
+		// Make a volume
+		FSpatialReadinessHandle Handle = SimCallback->AddReadinessVolume(UnitBounds, TEXT("Test Volume"));
+		AdvanceAndWait(Solver);
+		REQUIRE(SimCallback->GetUnreadyRigidParticles_PT().Num() == 0);
+
+		// Make a particle which should overlap that volume
+		FSingleParticlePhysicsProxy* BoxProxy = MakeDynamicBox(Scene, UnitBounds, false);
+		AdvanceAndWait(Solver);
+
+		// Advance one more tick to detect the midphase
+		AdvanceAndWait(Solver);
+		REQUIRE(SimCallback->GetUnreadyRigidParticles_PT().Num() == 1);
+	}
+
+	/*
+	* NOTE: This test will fail until  
+	* 
+	SECTION("Spawn dynamic particle overlapping unready volume, with gravity")
+	{
+		// Make a volume
+		FSpatialReadinessHandle Handle = SimCallback->AddReadinessVolume(UnitBounds, TEXT("Test Volume"));
+		AdvanceAndWait(Solver);
+		REQUIRE(SimCallback->GetParticleDataCache_PT().Num() == 0);
+
+		// Make a particle which should overlap that volume
+		FSingleParticlePhysicsProxy* BoxProxy = MakeDynamicBox(Scene, UnitBounds, true);
+		AdvanceAndWait(Solver);
+
+		// Advance one more tick to detect the midphase
+		AdvanceAndWait(Solver);
+		REQUIRE(SimCallback->GetParticleDataCache_PT().Num() == 1);
+	}
+	*/
+
+	SECTION("Dynamic particle with gravity stops falling when it hits an unready volume")
+	{
+		// Make a physics thread scope to avoid any thread context checks
+		FPhysicsThreadContextScope PTScope(true);
+
+		// Make a volume
+		FBox UnreadyBox(FVector(-1000, -1000, -10000), FVector(1000, 1000, 0));
+		FSpatialReadinessHandle Handle = SimCallback->AddReadinessVolume(UnreadyBox, TEXT("Test Volume"));
+		AdvanceAndWait(Solver);
+		REQUIRE(SimCallback->GetUnreadyRigidParticles_PT().Num() == 0);
+
+		// Make a particle which should overlap that volume
+		FBox FallingBox(FVector(-1, -1, 2), FVector(1, 1, 3));
+		FSingleParticlePhysicsProxy* BoxProxy = MakeDynamicBox(Scene, FallingBox, true);
+		AdvanceAndWait(Solver);
+
+		// Advance one more tick to detect the midphase and record the position of the particle
+		AdvanceAndWait(Solver);
+		REQUIRE(SimCallback->GetUnreadyRigidParticles_PT().Num() == 1);
+		const FVec3 FrozenX0 = BoxProxy->GetPhysicsThreadAPI()->X();
+
+		// Advance again and make sure the box didn't move. Do it 10 times for good measure
+		for (int32 Iteration = 0; Iteration < 10; ++Iteration)
+		{
+			AdvanceAndWait(Solver);
+			const FVec3 FrozenX1 = BoxProxy->GetPhysicsThreadAPI()->X();
+			REQUIRE(FrozenX0 == FrozenX1);
+		}
+
+		// Mark the volume handle as "ready" and advance again - the particle
+		// should be removed from the list of unready particles
+		Handle.MarkReady();
+		AdvanceAndWait(Solver);
+		REQUIRE(SimCallback->GetUnreadyRigidParticles_PT().Num() == 0);
+
+		// Advancing one more frame, we should see the particle start to fall
+		// again.
+		AdvanceAndWait(Solver);
+		const FVec3 FrozenX2 = BoxProxy->GetPhysicsThreadAPI()->X();
+		REQUIRE(FrozenX2.Z < FrozenX0.Z);
+	}
+}
diff --git a/Engine/Source/Runtime/Experimental/Chaos/Public/Chaos/Framework/PhysicsSolverBase.h b/Engine/Source/Runtime/Experimental/Chaos/Public/Chaos/Framework/PhysicsSolverBase.h
--- a/Engine/Source/Runtime/Experimental/Chaos/Public/Chaos/Framework/PhysicsSolverBase.h
+++ b/Engine/Source/Runtime/Experimental/Chaos/Public/Chaos/Framework/PhysicsSolverBase.h
@@ -368,10 +368,10 @@ namespace Chaos
 		}
 
 		/** Creates a new sim callback object of the type given. Caller expected to free using FreeSimCallbackObject_External*/
-		template <typename TSimCallbackObjectType>
-		inline TSimCallbackObjectType* CreateAndRegisterSimCallbackObject_External()
+		template <typename TSimCallbackObjectType, typename... ConstructorArgTypes>
+		inline TSimCallbackObjectType* CreateAndRegisterSimCallbackObject_External(ConstructorArgTypes&&... ConstructorArgs)
 		{
-			auto NewCallbackObject = new TSimCallbackObjectType();
+			auto NewCallbackObject = new TSimCallbackObjectType(Forward<ConstructorArgTypes>(ConstructorArgs)...);
 			//If at least one callback is on frozen GT we mark it as such
 			//Note this is never cleaned up (to avoid race conditions of register/unregister).
 			//Expectation is that there's typically just one of these on frozen gt that triggers all the fixed tick on gt (maybe two if another system wants direct control)
