Use SimSweep to do an initial readiness query for newly spawned particles, so that they can be frozen before their initial integration step. Also add cycle stats for each major part of the physics spatial readiness

diff --git a/Engine/Plugins/Experimental/SpatialReadiness/Source/PhysicsSpatialReadiness/Private/PhysicsSpatialReadiness/PhysicsSpatialReadinessSimCallback.cpp b/Engine/Plugins/Experimental/SpatialReadiness/Source/PhysicsSpatialReadiness/Private/PhysicsSpatialReadiness/PhysicsSpatialReadinessSimCallback.cpp
--- a/Engine/Plugins/Experimental/SpatialReadiness/Source/PhysicsSpatialReadiness/Private/PhysicsSpatialReadiness/PhysicsSpatialReadinessSimCallback.cpp
+++ b/Engine/Plugins/Experimental/SpatialReadiness/Source/PhysicsSpatialReadiness/Private/PhysicsSpatialReadiness/PhysicsSpatialReadinessSimCallback.cpp
@@ -4,12 +4,14 @@
 #include "PhysicsProxy/SingleParticlePhysicsProxy.h"
 #include "Chaos/ChaosEngineInterface.h"
 #include "Chaos/ChaosScene.h"
+#include "Chaos/Collision/SimSweep.h"
 #include "Chaos/MidPhaseModification.h"
 #include "Chaos/ParticleHandle.h"
 #include "SpatialReadinessHandle.h"
 #include "PBDRigidsSolver.h"
 #include "Chaos/PBDRigidsEvolutionGBF.h"
 #include "Physics/PhysicsFiltering.h"
+#include "PhysicsSpatialReadinessStats.h"
 
 using namespace Chaos;
 
@@ -30,6 +32,8 @@ FPBDRigidsEvolution* FPhysicsSpatialReadinessSimCallback::GetEvolution()
 
 int32 FPhysicsSpatialReadinessSimCallback::AddUnreadyVolume_GT(const FBox& Bounds, const FString& Description)
 {
+	SCOPE_CYCLE_COUNTER(STAT_SpatialReadiness_Physics_AddUnreadyVolumeGT)
+
 	// Get a pointer to the input struct - if this is null just cancel. If we create a
 	// particle without the input, then we'll lose track of it.
 	FPhysicsSpatialReadinessSimCallbackInput* SimInputPtr = GetProducerInputData_External();
@@ -103,6 +107,8 @@ int32 FPhysicsSpatialReadinessSimCallback::AddUnreadyVolume_GT(const FBox& Bound
 
 void FPhysicsSpatialReadinessSimCallback::RemoveUnreadyVolume_GT(int32 UnreadyVolumeIndex)
 {
+	SCOPE_CYCLE_COUNTER(STAT_SpatialReadiness_Physics_RemoveUnreadyVolumeGT)
+
 	if (!ensureMsgf(UnreadyVolumeParticles_GT.IsValidIndex(UnreadyVolumeIndex), TEXT("Trying to remove unready volume whose index is not being tracked")))
 	{
 		return;
@@ -132,6 +138,8 @@ void FPhysicsSpatialReadinessSimCallback::RemoveUnreadyVolume_GT(int32 UnreadyVo
 
 void FPhysicsSpatialReadinessSimCallback::FreezeParticles_PT()
 {
+	SCOPE_CYCLE_COUNTER(STAT_SpatialReadiness_Physics_FreezeParticlesPT)
+
 	FPBDRigidsEvolution* Evolution = GetEvolution();
 	if (!ensureMsgf(Evolution, TEXT("Attempted to freeze particle, but had no evolution")))
 	{
@@ -159,6 +167,8 @@ void FPhysicsSpatialReadinessSimCallback::FreezeParticles_PT()
 
 void FPhysicsSpatialReadinessSimCallback::UnFreezeParticles_PT()
 {
+	SCOPE_CYCLE_COUNTER(STAT_SpatialReadiness_Physics_UnFreezeParticlesPT)
+
 	FPBDRigidsEvolution* Evolution = GetEvolution();
 	if (!ensureMsgf(Evolution, TEXT("Attempted to un-freeze particles, but had no evolution")))
 	{
@@ -184,6 +194,8 @@ void FPhysicsSpatialReadinessSimCallback::UnFreezeParticles_PT()
 
 void FPhysicsSpatialReadinessSimCallback::OnPreSimulate_Internal()
 {
+	SCOPE_CYCLE_COUNTER(STAT_SpatialReadiness_Physics_PreSimulate)
+
 	// Process inputs we may have gotten from the game thread
 	if (const FPhysicsSpatialReadinessSimCallbackInput* Input = GetConsumerInput_Internal())
 	{
@@ -201,8 +213,91 @@ void FPhysicsSpatialReadinessSimCallback::OnPreSimulate_Internal()
 	}
 }
 
+void FPhysicsSpatialReadinessSimCallback::OnParticlesRegistered_Internal(TArray<FSingleParticlePhysicsProxy*>& RegisteredProxies)
+{
+	SCOPE_CYCLE_COUNTER(STAT_SpatialReadiness_Physics_ParticlesRegistered)
+
+	// Get the evolution, and get the acceleration structure from it
+	FPBDRigidsEvolution* Evolution = GetEvolution();
+	if (Evolution == nullptr)
+	{
+		return;
+	}
+
+	// For each newly added particle, query against unready volumes to see if
+	// it should be frozen.
+	for (FSingleParticlePhysicsProxy* ParticleProxy : RegisteredProxies)
+	{
+		FGeometryParticleHandle* GeometryParticle = ParticleProxy->GetHandle_LowLevel();
+		if (GeometryParticle == nullptr)
+		{
+			continue;
+		}
+
+		FPBDRigidParticleHandle* RigidParticle = GeometryParticle->CastToRigidParticle();
+		if (RigidParticle == nullptr)
+		{
+			continue;
+		}
+
+		// Get the bounds to use for the query
+		const FAABB3 QueryBounds = RigidParticle->WorldSpaceInflatedBounds();
+
+		// Set up particle filters for the query, so that we only get unready volumes
+		// NOTE: If we could speed up this bit here, that'd be great
+		const auto& ParticleFilter = [this](const FGeometryParticleHandle* Particle) -> bool
+		{
+			// If the particle has a single particle physics proxy (and it does,
+			// I can almost guarantee it), then count it as a hit if it's in our
+			// list of unready volume particles.
+			const IPhysicsProxyBase* Proxy = Particle->PhysicsProxy();
+			if (Proxy->GetType() == EPhysicsProxyType::SingleParticleProxy)
+			{
+				return UnreadyVolumeParticles_PT.Contains(
+					static_cast<const Chaos::FSingleParticlePhysicsProxy*>(Proxy));
+			}
+			return false;
+		};
+
+		// We know that unready volumes will all be boxes, so in theory we could
+		// probably filter a bunch of interactions before getting to the particle
+		// filter, however the particle filter is applied first so there's no point
+		// in implementing a shape filter.
+		const auto& ShapeFilter = [](const FPerShapeData* Shape, const FImplicitObject* Implicit)
+		{
+			return true;
+		};
+
+		// Make a lambda for collecting 
+		TArray<Private::FSimOverlapParticleShape> Overlaps;
+		const auto& OverlapCollector = [&Overlaps](const Private::FSimOverlapParticleShape& Overlap)
+		{
+			Overlaps.Add(Overlap);
+		};
+
+		// Get the broadphase and query against it to see if this new rigid particle
+		// generates midphases with any unready volumes.
+		//
+		// NOTE: In theory we could have a first-hit version of this function which
+		// would potentially avoid some unnecessary checking.
+		//
+		// NOTE: The above note is only valid until we actually start registering
+		// freeze locks with unready volumes...
+		Private::SimOverlapBounds(Evolution->GetSpatialAcceleration(), QueryBounds, ParticleFilter, ShapeFilter, OverlapCollector);
+
+		// The overlaps have already been filtered for unready volume interaction,
+		// so if there were _any_ overlaps, all we need to do is freeze the particle
+		if (Overlaps.Num() > 0)
+		{
+			UnreadyRigidParticles_PT.Add(RigidParticle);
+		}
+	}
+}
+
 void FPhysicsSpatialReadinessSimCallback::OnMidPhaseModification_Internal(FMidPhaseModifierAccessor& Accessor)
 {
+	SCOPE_CYCLE_COUNTER(STAT_SpatialReadiness_Physics_MidPhase)
+
 	// Clear the list of unready particles
 	UnreadyRigidParticles_PT.Reset();
 
diff --git a/Engine/Plugins/Experimental/SpatialReadiness/Source/PhysicsSpatialReadiness/Public/PhysicsSpatialReadinessSimCallback.h b/Engine/Plugins/Experimental/SpatialReadiness/Source/PhysicsSpatialReadiness/Public/PhysicsSpatialReadinessSimCallback.h
--- a/Engine/Plugins/Experimental/SpatialReadiness/Source/PhysicsSpatialReadiness/Public/PhysicsSpatialReadinessSimCallback.h
+++ b/Engine/Plugins/Experimental/SpatialReadiness/Source/PhysicsSpatialReadiness/Public/PhysicsSpatialReadinessSimCallback.h
@@ -33,6 +33,7 @@ struct FPhysicsSpatialReadinessSimCallback
 		FPhysicsSpatialReadinessSimCallbackInput,
 		FSimCallbackNoOutput,
 		ESimCallbackOptions::Presimulate |
+		ESimCallbackOptions::ParticleRegister |
 		ESimCallbackOptions::MidPhaseModification |
 		ESimCallbackOptions::PreIntegrate |
 		ESimCallbackOptions::PostIntegrate>
@@ -51,6 +52,7 @@ public:
 	/* begin: TSimCallbackObject */
 protected:
 	virtual void OnPreSimulate_Internal() override;
+	virtual void OnParticlesRegistered_Internal(TArray<FSingleParticlePhysicsProxy*>& RegisteredProxies) override;
 	virtual void OnMidPhaseModification_Internal(FMidPhaseModifierAccessor& Accessor) override;
 	virtual void OnPreIntegrate_Internal() override;
 	virtual void OnPostIntegrate_Internal() override;
diff --git a/Engine/Plugins/Experimental/SpatialReadiness/Source/PhysicsSpatialReadiness/Public/PhysicsSpatialReadinessStats.h b/Engine/Plugins/Experimental/SpatialReadiness/Source/PhysicsSpatialReadiness/Public/PhysicsSpatialReadinessStats.h
new file mode 100644
--- /dev/null
+++ b/Engine/Plugins/Experimental/SpatialReadiness/Source/PhysicsSpatialReadiness/Public/PhysicsSpatialReadinessStats.h
@@ -0,0 +1,14 @@
+// Copyright Epic Games, Inc. All Rights Reserved.
+#pragma once
+#include "SpatialReadinessStats.h"
+
+DECLARE_CYCLE_STAT(TEXT("SpatialReadiness_Physics_AddUnreadyVolumeGT"), STAT_SpatialReadiness_Physics_AddUnreadyVolumeGT, STATGROUP_SpatialReadiness);
+DECLARE_CYCLE_STAT(TEXT("SpatialReadiness_Physics_RemoveUnreadyVolumeGT"), STAT_SpatialReadiness_Physics_RemoveUnreadyVolumeGT, STATGROUP_SpatialReadiness);
+DECLARE_CYCLE_STAT(TEXT("SpatialReadiness_Physics_FreezeParticlesPT"), STAT_SpatialReadiness_Physics_FreezeParticlesPT, STATGROUP_SpatialReadiness);
+DECLARE_CYCLE_STAT(TEXT("SpatialReadiness_Physics_UnFreezeParticlesPT"), STAT_SpatialReadiness_Physics_UnFreezeParticlesPT, STATGROUP_SpatialReadiness);
+DECLARE_CYCLE_STAT(TEXT("SpatialReadiness_Physics_PreSimulate"), STAT_SpatialReadiness_Physics_PreSimulate, STATGROUP_SpatialReadiness);
+DECLARE_CYCLE_STAT(TEXT("SpatialReadiness_Physics_ParticlesRegistered"), STAT_SpatialReadiness_Physics_ParticlesRegistered, STATGROUP_SpatialReadiness);
+DECLARE_CYCLE_STAT(TEXT("SpatialReadiness_Physics_MidPhase"), STAT_SpatialReadiness_Physics_MidPhase, STATGROUP_SpatialReadiness);
+
+
+
diff --git a/Engine/Plugins/Experimental/SpatialReadiness/Source/SpatialReadiness/Private/SpatialReadiness/SpatialReadinessStats.h b/Engine/Plugins/Experimental/SpatialReadiness/Source/SpatialReadiness/Private/SpatialReadiness/SpatialReadinessStats.h
deleted file mode 100644
diff --git a/Engine/Plugins/Experimental/SpatialReadiness/Source/SpatialReadiness/Public/SpatialReadinessStats.h b/Engine/Plugins/Experimental/SpatialReadiness/Source/SpatialReadiness/Public/SpatialReadinessStats.h
new file mode 100644
--- /dev/null
+++ b/Engine/Plugins/Experimental/SpatialReadiness/Source/SpatialReadiness/Public/SpatialReadinessStats.h
@@ -0,0 +1,12 @@
+// Copyright Epic Games, Inc. All Rights Reserved.
+
+#pragma once
+
+#include "Stats/Stats.h"
+
+DECLARE_STATS_GROUP(TEXT("SpatialReadiness"), STATGROUP_SpatialReadiness, STATCAT_Advanced);
+DECLARE_CYCLE_STAT(TEXT(""), STAT_SpatialReadiness_AddUnreadyVolume, STATGROUP_SpatialReadiness);
+DECLARE_CYCLE_STAT(TEXT(""), STAT_SpatialReadiness_RemoveUnreadyVolume, STATGROUP_SpatialReadiness);
+
+
+
diff --git a/Engine/Plugins/Experimental/SpatialReadiness/Tests/Private/PhysicsSpatialReadinessTests.cpp b/Engine/Plugins/Experimental/SpatialReadiness/Tests/Private/PhysicsSpatialReadinessTests.cpp
--- a/Engine/Plugins/Experimental/SpatialReadiness/Tests/Private/PhysicsSpatialReadinessTests.cpp
+++ b/Engine/Plugins/Experimental/SpatialReadiness/Tests/Private/PhysicsSpatialReadinessTests.cpp
@@ -151,15 +151,12 @@ TEST_CASE("Sim Callback", "[physics spatial readiness]")
 		REQUIRE(SimCallback->GetUnreadyRigidParticles_PT().Num() == 1);
 	}
 
-	/*
-	* NOTE: This test will fail until  
-	* 
 	SECTION("Spawn dynamic particle overlapping unready volume, with gravity")
 	{
 		// Make a volume
 		FSpatialReadinessHandle Handle = SimCallback->AddReadinessVolume(UnitBounds, TEXT("Test Volume"));
 		AdvanceAndWait(Solver);
-		REQUIRE(SimCallback->GetParticleDataCache_PT().Num() == 0);
+		REQUIRE(SimCallback->GetUnreadyRigidParticles_PT().Num() == 0);
 
 		// Make a particle which should overlap that volume
 		FSingleParticlePhysicsProxy* BoxProxy = MakeDynamicBox(Scene, UnitBounds, true);
@@ -167,9 +164,8 @@ TEST_CASE("Sim Callback", "[physics spatial readiness]")
 
 		// Advance one more tick to detect the midphase
 		AdvanceAndWait(Solver);
-		REQUIRE(SimCallback->GetParticleDataCache_PT().Num() == 1);
+		REQUIRE(SimCallback->GetUnreadyRigidParticles_PT().Num() == 1);
 	}
-	*/
 
 	SECTION("Dynamic particle with gravity stops falling when it hits an unready volume")
 	{
