Renamed FSpatialReadinessHandle to FSpatialReadinessVolume. Renamed FPhysicsSpatialReadinessSubsystem to FSpatialReadinessSubsystem, and removed the PhysicsSpatialReadiness module, combining them into one.

diff --git a/Engine/Plugins/Experimental/SpatialReadiness/Source/PhysicsSpatialReadiness/PhysicsSpatialReadiness.Build.cs b/Engine/Plugins/Experimental/SpatialReadiness/Source/PhysicsSpatialReadiness/PhysicsSpatialReadiness.Build.cs
deleted file mode 100644
diff --git a/Engine/Plugins/Experimental/SpatialReadiness/Source/PhysicsSpatialReadiness/Private/PhysicsSpatialReadiness/PhysicsSpatialReadinessModule.cpp b/Engine/Plugins/Experimental/SpatialReadiness/Source/PhysicsSpatialReadiness/Private/PhysicsSpatialReadiness/PhysicsSpatialReadinessModule.cpp
deleted file mode 100644
diff --git a/Engine/Plugins/Experimental/SpatialReadiness/Source/PhysicsSpatialReadiness/Private/PhysicsSpatialReadiness/PhysicsSpatialReadinessSimCallback.cpp b/Engine/Plugins/Experimental/SpatialReadiness/Source/PhysicsSpatialReadiness/Private/PhysicsSpatialReadiness/PhysicsSpatialReadinessSimCallback.cpp
deleted file mode 100644
diff --git a/Engine/Plugins/Experimental/SpatialReadiness/Source/PhysicsSpatialReadiness/Private/PhysicsSpatialReadiness/PhysicsSpatialReadinessSubsystem.cpp b/Engine/Plugins/Experimental/SpatialReadiness/Source/PhysicsSpatialReadiness/Private/PhysicsSpatialReadiness/PhysicsSpatialReadinessSubsystem.cpp
deleted file mode 100644
diff --git a/Engine/Plugins/Experimental/SpatialReadiness/Source/PhysicsSpatialReadiness/Public/PhysicsSpatialReadinessModule.h b/Engine/Plugins/Experimental/SpatialReadiness/Source/PhysicsSpatialReadiness/Public/PhysicsSpatialReadinessModule.h
deleted file mode 100644
diff --git a/Engine/Plugins/Experimental/SpatialReadiness/Source/PhysicsSpatialReadiness/Public/PhysicsSpatialReadinessSimCallback.h b/Engine/Plugins/Experimental/SpatialReadiness/Source/PhysicsSpatialReadiness/Public/PhysicsSpatialReadinessSimCallback.h
deleted file mode 100644
diff --git a/Engine/Plugins/Experimental/SpatialReadiness/Source/PhysicsSpatialReadiness/Public/PhysicsSpatialReadinessStats.h b/Engine/Plugins/Experimental/SpatialReadiness/Source/PhysicsSpatialReadiness/Public/PhysicsSpatialReadinessStats.h
deleted file mode 100644
diff --git a/Engine/Plugins/Experimental/SpatialReadiness/Source/PhysicsSpatialReadiness/Public/PhysicsSpatialReadinessSubsystem.h b/Engine/Plugins/Experimental/SpatialReadiness/Source/PhysicsSpatialReadiness/Public/PhysicsSpatialReadinessSubsystem.h
deleted file mode 100644
diff --git a/Engine/Plugins/Experimental/SpatialReadiness/Source/SpatialReadiness/Private/SpatialReadiness/SpatialReadinessAPI.cpp b/Engine/Plugins/Experimental/SpatialReadiness/Source/SpatialReadiness/Private/SpatialReadiness/SpatialReadinessAPI.cpp
--- a/Engine/Plugins/Experimental/SpatialReadiness/Source/SpatialReadiness/Private/SpatialReadiness/SpatialReadinessAPI.cpp
+++ b/Engine/Plugins/Experimental/SpatialReadiness/Source/SpatialReadiness/Private/SpatialReadiness/SpatialReadinessAPI.cpp
@@ -21,10 +21,10 @@ FSpatialReadinessAPI::~FSpatialReadinessAPI()
 	}
 }
 
-FSpatialReadinessHandle FSpatialReadinessAPI::CreateHandle(const FBox& Bounds, const FString& Description)
+FSpatialReadinessVolume FSpatialReadinessAPI::CreateHandle(const FBox& Bounds, const FString& Description)
 {
 	// Create a handle for a new unready volume
-	FSpatialReadinessHandle Handle(Delegates.ToWeakPtr(), Bounds, Description);
+	FSpatialReadinessVolume Handle(Delegates.ToWeakPtr(), Bounds, Description);
 
 	// Default to an unready state. This should call back into AddUnreadyVolume_Internal
 	// to create the underlying unready volume.
diff --git a/Engine/Plugins/Experimental/SpatialReadiness/Source/SpatialReadiness/Private/SpatialReadiness/SpatialReadinessHandle.cpp b/Engine/Plugins/Experimental/SpatialReadiness/Source/SpatialReadiness/Private/SpatialReadiness/SpatialReadinessHandle.cpp
deleted file mode 100644
diff --git a/Engine/Plugins/Experimental/SpatialReadiness/Source/SpatialReadiness/Private/SpatialReadiness/SpatialReadinessSimCallback.cpp b/Engine/Plugins/Experimental/SpatialReadiness/Source/SpatialReadiness/Private/SpatialReadiness/SpatialReadinessSimCallback.cpp
new file mode 100644
--- /dev/null
+++ b/Engine/Plugins/Experimental/SpatialReadiness/Source/SpatialReadiness/Private/SpatialReadiness/SpatialReadinessSimCallback.cpp
@@ -0,0 +1,351 @@
+// Copyright Epic Games, Inc. All Rights Reserved.
+
+#include "SpatialReadinessSimCallback.h"
+#include "PhysicsProxy/SingleParticlePhysicsProxy.h"
+#include "Chaos/ChaosEngineInterface.h"
+#include "Chaos/ChaosScene.h"
+#include "Chaos/Collision/SimSweep.h"
+#include "Chaos/MidPhaseModification.h"
+#include "Chaos/ParticleHandle.h"
+#include "SpatialReadinessVolume.h"
+#include "PBDRigidsSolver.h"
+#include "Chaos/PBDRigidsEvolutionGBF.h"
+#include "Physics/PhysicsFiltering.h"
+#include "SpatialReadinessStats.h"
+
+using namespace Chaos;
+
+FSpatialReadinessSimCallback::FSpatialReadinessSimCallback(FChaosScene& InPhysicsScene)
+	: PhysicsScene(InPhysicsScene)
+	, ParticleDataCache_PT()
+{ }
+
+FPBDRigidsEvolution* FSpatialReadinessSimCallback::GetEvolution()
+{
+	if (FPBDRigidsSolver* MySolver = static_cast<FPBDRigidsSolver*>(GetSolver()))
+	{
+		return MySolver->GetEvolution();
+	}
+
+	return nullptr;
+}
+
+int32 FSpatialReadinessSimCallback::AddUnreadyVolume_GT(const FBox& Bounds, const FString& Description)
+{
+	SCOPE_CYCLE_COUNTER(STAT_SpatialReadiness_Physics_AddUnreadyVolumeGT)
+
+	// Get a pointer to the input struct - if this is null just cancel. If we create a
+	// particle without the input, then we'll lose track of it.
+	FSpatialReadinessSimCallbackInput* SimInputPtr = GetProducerInputData_External();
+	if (!ensureMsgf(SimInputPtr, TEXT("Failed to access sim callback object input")))
+	{
+		return INDEX_NONE;
+	}
+	FSpatialReadinessSimCallbackInput& SimInput = *SimInputPtr;
+
+	// Create a box implicit geometry from the bounds
+	const FVec3 BoxCenter = (Bounds.Min + Bounds.Max) * .5f;
+	const FVec3 BoxHalfExtent = Bounds.Max - Bounds.Min;
+	Chaos::FImplicitObjectPtr BoxGeom = MakeImplicitObjectPtr<TBox<FReal, 3>>(-BoxHalfExtent, BoxHalfExtent);
+
+	// Create a new static particle to represent the volume
+	FActorCreationParams Params;
+	Params.bSimulatePhysics = false;
+	Params.bStatic = true;
+	Params.InitialTM = FTransform(FQuat::Identity, BoxCenter);
+	Params.Scene = &PhysicsScene;
+	FSingleParticlePhysicsProxy* ParticleProxy = nullptr;
+	FChaosEngineInterface::CreateActor(Params, ParticleProxy);
+	if (!ensureMsgf(ParticleProxy, TEXT("Failed to create new particle proxy")))
+	{
+		return INDEX_NONE;
+	}
+	FRigidBodyHandle_External& ParticleHandle = ParticleProxy->GetGameThreadAPI();
+
+	// Create collision response container
+	FCollisionResponseContainer CollisionResponse;
+	CollisionResponse.SetAllChannels(ECollisionResponse::ECR_Block);
+
+	// Create collision filter data for the particle
+	FCollisionFilterData QueryFilterData, SimFilterData;
+	CreateShapeFilterData(
+		/* MyChannel */            static_cast<uint8>(ECollisionChannel::ECC_WorldStatic),
+		/* MaskFilter */           FMaskFilter(0),
+		/* SourceObjectID */       0,
+		/* ResponseToChannels */   CollisionResponse,
+		/* ComponentID */          0,
+		/* BodyIndex */            0,
+		/* OutQueryData */         QueryFilterData,
+		/* OutSimData */           SimFilterData,
+		/* bEnableCCD */           true,
+		/* bEnableContactNotify */ false,
+		/* bStaticShape */         true);
+
+	// Make the geometry
+	ParticleHandle.SetGeometry(BoxGeom);
+	ParticleHandle.SetShapeSimCollisionEnabled(0, true);
+	ParticleHandle.SetShapeQueryCollisionEnabled(0, false);
+	ParticleHandle.SetShapeSimData(0, SimFilterData);
+	ParticleHandle.ShapesArray()[0]->SetIsProbe(true);
+#if CHAOS_DEBUG_NAME
+	ParticleHandle.SetDebugName(MakeShared<FString>(TEXT("Unready Volume")));
+#endif
+
+	// Add the new particle to the scene
+	TArray<FPhysicsActorHandle> Actors = { ParticleProxy };
+	PhysicsScene.AddActorsToScene_AssumesLocked(Actors);
+
+	// Save the proxy in our list of GT particles
+	const int32 ParticleIndex = UnreadyVolumeParticles_GT.Add(ParticleProxy);
+
+	// Queue up the particle proxy for processing on PT
+	SimInput.UnreadyVolumesToAdd.Add(ParticleProxy);
+	SimInput.UnreadyVolumesToRemove.Remove(ParticleProxy);
+
+	return ParticleIndex;
+}
+
+void FSpatialReadinessSimCallback::RemoveUnreadyVolume_GT(int32 UnreadyVolumeIndex)
+{
+	SCOPE_CYCLE_COUNTER(STAT_SpatialReadiness_Physics_RemoveUnreadyVolumeGT)
+
+	if (!ensureMsgf(UnreadyVolumeParticles_GT.IsValidIndex(UnreadyVolumeIndex), TEXT("Trying to remove unready volume whose index is not being tracked")))
+	{
+		return;
+	}
+
+	// Get the proxy associated with this index
+	FSingleParticlePhysicsProxy* ParticleProxy = UnreadyVolumeParticles_GT[UnreadyVolumeIndex];
+	if (!ensureMsgf(ParticleProxy, TEXT("Particle proxy associated with unready volume index was null")))
+	{
+		return;
+	}
+	
+	// Free the index in our GT tracker
+	UnreadyVolumeParticles_GT.RemoveAt(UnreadyVolumeIndex);
+
+	// Tell the PT to remove it's tracking of this proxy as well
+	FSpatialReadinessSimCallbackInput* SimInput = GetProducerInputData_External();
+	if (ensureMsgf(SimInput, TEXT("Failed to access sim input data")))
+	{
+		SimInput->UnreadyVolumesToRemove.Add(ParticleProxy);
+		SimInput->UnreadyVolumesToAdd.Remove(ParticleProxy);
+	}
+
+	// Delete the particle
+	FChaosEngineInterface::ReleaseActor(ParticleProxy, &PhysicsScene);
+}
+
+void FSpatialReadinessSimCallback::FreezeParticles_PT()
+{
+	SCOPE_CYCLE_COUNTER(STAT_SpatialReadiness_Physics_FreezeParticlesPT)
+
+	FPBDRigidsEvolution* Evolution = GetEvolution();
+	if (!ensureMsgf(Evolution, TEXT("Attempted to freeze particle, but had no evolution")))
+	{
+		return;
+	}
+
+	for (FPBDRigidParticleHandle* RigidParticle : UnreadyRigidParticles_PT)
+	{
+		const EObjectStateType ObjectState = RigidParticle->ObjectState();
+		if (ObjectState == EObjectStateType::Static)
+		{
+			return;
+		}
+
+		// Get the geometry particle
+		FGeometryParticleHandle* GeometryParticle = RigidParticle;
+
+		// Cache the current object state
+		ParticleDataCache_PT.Add({ GeometryParticle, ObjectState });
+
+		// Set the object state to static
+		Evolution->SetParticleObjectState(RigidParticle, EObjectStateType::Static);
+	}
+}
+
+void FSpatialReadinessSimCallback::UnFreezeParticles_PT()
+{
+	SCOPE_CYCLE_COUNTER(STAT_SpatialReadiness_Physics_UnFreezeParticlesPT)
+
+	FPBDRigidsEvolution* Evolution = GetEvolution();
+	if (!ensureMsgf(Evolution, TEXT("Attempted to un-freeze particles, but had no evolution")))
+	{
+		return;
+	}
+
+	// For each particle that we froze, unfreeze it
+	for (TPair<FGeometryParticleHandle*, EObjectStateType>& ParticleData : ParticleDataCache_PT)
+	{
+		FPBDRigidParticleHandle* RigidParticle = ParticleData.Key->CastToRigidParticle();
+		if (RigidParticle == nullptr)
+		{
+			continue;
+		}
+
+		// Restore the object state of this particle
+		Evolution->SetParticleObjectState(RigidParticle, ParticleData.Value);
+	}
+
+	// Clear the array
+	ParticleDataCache_PT.Reset();
+}
+
+void FSpatialReadinessSimCallback::OnPreSimulate_Internal()
+{
+	SCOPE_CYCLE_COUNTER(STAT_SpatialReadiness_Physics_PreSimulate)
+
+	// Process inputs we may have gotten from the game thread
+	if (const FSpatialReadinessSimCallbackInput* Input = GetConsumerInput_Internal())
+	{
+		// Process additions
+		for (FSingleParticlePhysicsProxy* ParticleProxy : Input->UnreadyVolumesToAdd)
+		{
+			UnreadyVolumeParticles_PT.Add(ParticleProxy);
+		}
+
+		// Process removals
+		for (FSingleParticlePhysicsProxy* ParticleProxy : Input->UnreadyVolumesToRemove)
+		{
+			UnreadyVolumeParticles_PT.Remove(ParticleProxy);
+		}
+	}
+}
+
+void FSpatialReadinessSimCallback::OnParticlesRegistered_Internal(TArray<FSingleParticlePhysicsProxy*>& RegisteredProxies)
+{
+	SCOPE_CYCLE_COUNTER(STAT_SpatialReadiness_Physics_ParticlesRegistered)
+
+	// Get the evolution, and get the acceleration structure from it
+	FPBDRigidsEvolution* Evolution = GetEvolution();
+	if (Evolution == nullptr)
+	{
+		return;
+	}
+
+	// For each newly added particle, query against unready volumes to see if
+	// it should be frozen.
+	for (FSingleParticlePhysicsProxy* ParticleProxy : RegisteredProxies)
+	{
+		FGeometryParticleHandle* GeometryParticle = ParticleProxy->GetHandle_LowLevel();
+		if (GeometryParticle == nullptr)
+		{
+			continue;
+		}
+
+		FPBDRigidParticleHandle* RigidParticle = GeometryParticle->CastToRigidParticle();
+		if (RigidParticle == nullptr)
+		{
+			continue;
+		}
+
+		// Get the bounds to use for the query
+		const FAABB3 QueryBounds = RigidParticle->WorldSpaceInflatedBounds();
+
+		// Set up particle filters for the query, so that we only get unready volumes
+		// NOTE: If we could speed up this bit here, that'd be great
+		const auto& ParticleFilter = [this](const FGeometryParticleHandle* Particle) -> bool
+		{
+			// If the particle has a single particle physics proxy (and it does,
+			// I can almost guarantee it), then count it as a hit if it's in our
+			// list of unready volume particles.
+			const IPhysicsProxyBase* Proxy = Particle->PhysicsProxy();
+			if (Proxy->GetType() == EPhysicsProxyType::SingleParticleProxy)
+			{
+				return UnreadyVolumeParticles_PT.Contains(
+					static_cast<const Chaos::FSingleParticlePhysicsProxy*>(Proxy));
+			}
+			return false;
+		};
+
+		// We know that unready volumes will all be boxes, so in theory we could
+		// probably filter a bunch of interactions before getting to the particle
+		// filter, however the particle filter is applied first so there's no point
+		// in implementing a shape filter.
+		const auto& ShapeFilter = [](const FPerShapeData* Shape, const FImplicitObject* Implicit)
+		{
+			return true;
+		};
+
+		// Make a lambda for collecting 
+		TArray<Private::FSimOverlapParticleShape> Overlaps;
+		const auto& OverlapCollector = [&Overlaps](const Private::FSimOverlapParticleShape& Overlap)
+		{
+			Overlaps.Add(Overlap);
+		};
+
+		// Get the broadphase and query against it to see if this new rigid particle
+		// generates midphases with any unready volumes.
+		//
+		// NOTE: In theory we could have a first-hit version of this function which
+		// would potentially avoid some unnecessary checking.
+		//
+		// NOTE: The above note is only valid until we actually start registering
+		// freeze locks with unready volumes...
+		Private::SimOverlapBounds(Evolution->GetSpatialAcceleration(), QueryBounds, ParticleFilter, ShapeFilter, OverlapCollector);
+
+		// The overlaps have already been filtered for unready volume interaction,
+		// so if there were _any_ overlaps, all we need to do is freeze the particle
+		if (Overlaps.Num() > 0)
+		{
+			UnreadyRigidParticles_PT.Add(RigidParticle);
+		}
+	}
+}
+
+void FSpatialReadinessSimCallback::OnMidPhaseModification_Internal(FMidPhaseModifierAccessor& Accessor)
+{
+	SCOPE_CYCLE_COUNTER(STAT_SpatialReadiness_Physics_MidPhase)
+
+	// Clear the list of unready particles
+	UnreadyRigidParticles_PT.Reset();
+
+	// Go through every unready volume
+	for (FSingleParticlePhysicsProxy* UnreadyProxy : UnreadyVolumeParticles_PT)
+	{
+		if (UnreadyProxy->GetMarkedDeleted())
+		{
+			continue;
+		}
+
+		FGeometryParticleHandle* UnreadyVolume = UnreadyProxy->GetHandle_LowLevel();
+		if (UnreadyVolume == nullptr)
+		{
+			continue;
+		}
+
+		// Go through every mid-phase which involves this volume
+		for (FMidPhaseModifier& MidPhase : Accessor.GetMidPhases(UnreadyVolume))
+		{
+			// Get the particle that is not the unready volume
+			FGeometryParticleHandle* GeometryParticle = MidPhase.GetOtherParticle(UnreadyVolume);
+			if (GeometryParticle == nullptr)
+			{
+				continue;
+			}
+
+			FPBDRigidParticleHandle* RigidParticle = GeometryParticle->CastToRigidParticle();
+			if (RigidParticle == nullptr)
+			{
+				continue;
+			}
+
+			// Add the particle to the list of particles to freeze
+			UnreadyRigidParticles_PT.Add(RigidParticle);
+
+			// Disable the midphase
+			MidPhase.Disable();
+		}
+	}
+}
+
+void FSpatialReadinessSimCallback::OnPreIntegrate_Internal()
+{
+	FreezeParticles_PT();
+}
+
+void FSpatialReadinessSimCallback::OnPostIntegrate_Internal()
+{
+	UnFreezeParticles_PT();
+}
diff --git a/Engine/Plugins/Experimental/SpatialReadiness/Source/SpatialReadiness/Private/SpatialReadiness/SpatialReadinessSubsystem.cpp b/Engine/Plugins/Experimental/SpatialReadiness/Source/SpatialReadiness/Private/SpatialReadiness/SpatialReadinessSubsystem.cpp
new file mode 100644
--- /dev/null
+++ b/Engine/Plugins/Experimental/SpatialReadiness/Source/SpatialReadiness/Private/SpatialReadiness/SpatialReadinessSubsystem.cpp
@@ -0,0 +1,150 @@
+// Copyright Epic Games, Inc. All Rights Reserved.
+
+#include "SpatialReadinessSubsystem.h"
+#include "Engine/World.h"
+#include "Physics/Experimental/PhysScene_Chaos.h"
+#include "PBDRigidsSolver.h"
+
+using namespace Chaos;
+
+USpatialReadinessSubsystem::USpatialReadinessSubsystem()
+	: Super()
+	, SpatialReadiness(this, &This::AddUnreadyVolume, &This::RemoveUnreadyVolume) { }
+
+USpatialReadinessSubsystem::USpatialReadinessSubsystem(FVTableHelper& Helper)
+	: Super(Helper)
+	, SpatialReadiness(this, &This::AddUnreadyVolume, &This::RemoveUnreadyVolume) { }
+
+bool USpatialReadinessSubsystem::ShouldCreateSubsystem(UObject* Outer) const
+{
+	// TODO: Where should the setting for this exist? Is it enough
+	//       to just not load the physics readiness module if we don't
+	//       want to use it?
+	return true;
+}
+
+void USpatialReadinessSubsystem::OnWorldBeginPlay(UWorld& InWorld)
+{
+	Super::OnWorldBeginPlay(InWorld);
+	CreateSimCallback();
+}
+
+void USpatialReadinessSubsystem::Initialize(FSubsystemCollectionBase& Collection)
+{
+	Super::Initialize(Collection);
+}
+
+void USpatialReadinessSubsystem::Deinitialize()
+{
+	DestroySimCallback();
+	Super::Deinitialize();
+}
+
+FSpatialReadinessVolume USpatialReadinessSubsystem::AddReadinessVolume(const FBox& Bounds, const FString& Description)
+{
+	// Create a readiness volume and return it's "handle".
+	// This call will trigger the associated AddUnreadyVolume method
+	// since volumes are unready by default.
+	return SpatialReadiness.CreateHandle(Bounds, Description);
+}
+
+int32 USpatialReadinessSubsystem::AddUnreadyVolume(const FBox& Bounds, const FString& Description)
+{
+	// Add the physics unready volume
+	if (ensureMsgf(SimCallback, TEXT("Tried to add unready volume when no sim callback exists")))
+	{
+		return SimCallback->AddUnreadyVolume_GT(Bounds, Description);
+	}
+
+	return INDEX_NONE;
+}
+
+void USpatialReadinessSubsystem::RemoveUnreadyVolume(int32 UnreadyVolumeIndex)
+{
+	// Remove the physics unready volume
+	if (ensureMsgf(SimCallback, TEXT("Tried to add unready volume when no sim callback exists")))
+	{
+		SimCallback->RemoveUnreadyVolume_GT(UnreadyVolumeIndex);
+	}
+}
+
+bool USpatialReadinessSubsystem::CreateSimCallback()
+{
+	// If we already have a sim callback, destroy it.
+	if (SimCallback)
+	{
+		DestroySimCallback();
+	}
+
+	// If we still have a sim callback at this point, then we
+	// must have failed to delete it which means proceding would
+	// leave a dangling pointer.
+	if (!ensureMsgf(SimCallback == nullptr, TEXT("Tried and failed to destroy existing sim callback so that a new one could be created.")))
+	{
+		return false;
+	}
+
+	// The sim callback takes a scene reference in its constructor
+	FPhysScene_Chaos* Scene = GetScene();
+	if (!ensureMsgf(Scene, TEXT("Trying to create sim callback when there's no physics scene")))
+	{
+		return false;
+	}
+
+	// We need the solver to create the scene callback
+	Chaos::FPhysicsSolver* Solver = GetSolver();
+	if (!ensureMsgf(Solver, TEXT("Trying to create sim callback when there's no physics solver")))
+	{
+		return false;
+	}
+
+	// Request creation of the scene callback from the solver
+	SimCallback = Solver->CreateAndRegisterSimCallbackObject_External<FSpatialReadinessSimCallback>(static_cast<FChaosScene&>(*Scene));
+	if (!ensureMsgf(SimCallback, TEXT("Sim callback creation failed")))
+	{
+		return false;
+	}
+
+	// Return true to indicate successful creation of SimCallback
+	return true;
+}
+
+bool USpatialReadinessSubsystem::DestroySimCallback()
+{
+	if (SimCallback)
+	{
+		if (Chaos::FPhysicsSolver* Solver = GetSolver())
+		{
+			Solver->UnregisterAndFreeSimCallbackObject_External(SimCallback);
+			SimCallback = nullptr;
+			return true;
+		}
+	}
+
+	return false;
+}
+
+FPhysScene_Chaos* USpatialReadinessSubsystem::GetScene()
+{
+	if (UWorld* World = GetWorld())
+	{
+		if (FPhysScene* Scene = World->GetPhysicsScene())
+		{
+			return static_cast<FPhysScene_Chaos*>(Scene);
+		}
+	}
+	return nullptr;
+}
+
+Chaos::FPhysicsSolver* USpatialReadinessSubsystem::GetSolver()
+{
+	if (FPhysScene* Scene = GetScene())
+	{
+		if (Chaos::FPhysicsSolver* Solver = Scene->GetSolver())
+		{
+			return Solver;
+		}
+	}
+
+	return nullptr;
+}
diff --git a/Engine/Plugins/Experimental/SpatialReadiness/Source/SpatialReadiness/Private/SpatialReadiness/SpatialReadinessVolume.cpp b/Engine/Plugins/Experimental/SpatialReadiness/Source/SpatialReadiness/Private/SpatialReadiness/SpatialReadinessVolume.cpp
new file mode 100644
--- /dev/null
+++ b/Engine/Plugins/Experimental/SpatialReadiness/Source/SpatialReadiness/Private/SpatialReadiness/SpatialReadinessVolume.cpp
@@ -0,0 +1,136 @@
+// Copyright Epic Games, Inc. All Rights Reserved.
+
+#include "SpatialReadinessVolume.h"
+#include "SpatialReadinessAPIDelegates.h"
+
+FSpatialReadinessVolume::FSpatialReadinessVolume(
+	TWeakPtr<FSpatialReadinessAPIDelegates> InDelegates,
+	const FBox& InBounds,
+	const FString& InDescription)
+	: Delegates(InDelegates)
+	, Index(INDEX_NONE)
+	, Bounds(InBounds)
+#if WITH_SPATIALREADINESS_HANDLE_DESCRIPTIONS
+	, Description(InDescription)
+#endif
+{ }
+
+FSpatialReadinessVolume::FSpatialReadinessVolume(FSpatialReadinessVolume&& Other)
+	: Delegates(Other.Delegates)
+	, Index(Other.Index)
+	, Bounds(Other.Bounds)
+#if WITH_SPATIALREADINESS_HANDLE_DESCRIPTIONS
+	, Description(Other.Description)
+#endif
+{
+	Other.Delegates.Reset();
+	Other.Index = INDEX_NONE;
+	Other.Bounds = FBox();
+#if WITH_SPATIALREADINESS_HANDLE_DESCRIPTIONS
+	Other.Description = TEXT("");
+#endif
+}
+
+FSpatialReadinessVolume::~FSpatialReadinessVolume()
+{
+	// When a handle is destroyed, we mark the corresponding volume as
+	// "ready" (in case it wasn't already). This ensures that we do not
+	// end up losing track of unready volumes and leaving them hanging
+	// around in an acceleration structure.
+	//
+	// This is the only case where we actually do an IsValid check
+	// without _ensuring_ validity. If the api which created this
+	// handle has already been deleted, then MarkReady is not going to
+	// do anything anyway...
+	//
+	// I'm a bit conflicted on adding this check because if we ever
+	// come in here when IsValid == false, that means a handle was held
+	// onto for too long.
+	if (IsValid())
+	{
+		MarkReady();
+	}
+}
+
+bool FSpatialReadinessVolume::IsValid() const
+{
+	return Delegates.IsValid();
+}
+
+bool FSpatialReadinessVolume::EnsureIsValid() const
+{
+	return ensureMsgf(IsValid(), TEXT("Readiness volume handle with invalid delegates pointer is being accessed"));
+}
+
+bool FSpatialReadinessVolume::IsReady() const
+{
+	EnsureIsValid();
+
+	// If we have no index, that means there's no underlying unready volume,
+	// so we're "ready"
+	return Index == INDEX_NONE;
+}
+
+void FSpatialReadinessVolume::MarkReady()
+{
+	if (IsReady())
+	{
+		return;
+	}
+
+	TSharedPtr<FSpatialReadinessAPIDelegates> SharedDelegates = Delegates.Pin();
+	if (ensureMsgf(SharedDelegates, TEXT("Tried to mark ready a readiness volume handle with an invalid delegates container")))
+	{
+		// Remove the unready volume corresponding with our current index
+		SharedDelegates->RemoveUnreadyVolumeDelegate.ExecuteIfBound(Index);
+	}
+
+	// Stop tracking this index
+	Index = INDEX_NONE;
+}
+
+void FSpatialReadinessVolume::MarkUnready()
+{
+	if (!IsReady())
+	{
+		return;
+	}
+
+	TSharedPtr<FSpatialReadinessAPIDelegates> SharedDelegates = Delegates.Pin();
+	if (ensureMsgf(SharedDelegates, TEXT("Tried to mark unready a readiness volume handle with an invalid delegates container")))
+	{
+		// If we've got an index, that means we're already associated with an
+		// existing unready volume, which must be destroyed.
+		if (Index != INDEX_NONE)
+		{
+			// If we already had an unread volume, remove it so we can make
+			// a new one without leaving the old one dangling
+			SharedDelegates->RemoveUnreadyVolumeDelegate.ExecuteIfBound(Index);
+		}
+
+		// Update our index to a new unready volume
+		if (SharedDelegates->AddUnreadyVolumeDelegate.IsBound())
+		{
+			Index = SharedDelegates->AddUnreadyVolumeDelegate.Execute(Bounds, GetDescription());
+		}
+	}
+}
+
+void FSpatialReadinessVolume::MarkUnready(FBox& NewBounds)
+{
+	// Store new bounds
+	Bounds = NewBounds;
+
+	// Create a new unready volume with new bounds
+	MarkUnready();
+}
+
+const FString& FSpatialReadinessVolume::GetDescription() const
+{
+#if WITH_SPATIALREADINESS_HANDLE_DESCRIPTIONS
+	return Description;
+#else
+	static const FString EmptyString;
+	return EmptyString;
+#endif
+}
diff --git a/Engine/Plugins/Experimental/SpatialReadiness/Source/SpatialReadiness/Public/SpatialReadinessAPI.h b/Engine/Plugins/Experimental/SpatialReadiness/Source/SpatialReadiness/Public/SpatialReadinessAPI.h
--- a/Engine/Plugins/Experimental/SpatialReadiness/Source/SpatialReadiness/Public/SpatialReadinessAPI.h
+++ b/Engine/Plugins/Experimental/SpatialReadiness/Source/SpatialReadiness/Public/SpatialReadinessAPI.h
@@ -4,7 +4,7 @@
 
 #include "CoreMinimal.h"
 #include "Templates/Function.h"
-#include "SpatialReadinessHandle.h"
+#include "SpatialReadinessVolume.h"
 #include "SpatialReadinessSignatures.h"
 
 struct FSpatialReadinessAPIDelegates;
@@ -29,7 +29,7 @@ public:
 
 	SPATIALREADINESS_API ~FSpatialReadinessAPI();
 
-	SPATIALREADINESS_API FSpatialReadinessHandle CreateHandle(const FBox& Bounds, const FString& Description);
+	SPATIALREADINESS_API FSpatialReadinessVolume CreateHandle(const FBox& Bounds, const FString& Description);
 
 private:
 	TSharedPtr<FSpatialReadinessAPIDelegates> Delegates;
diff --git a/Engine/Plugins/Experimental/SpatialReadiness/Source/SpatialReadiness/Public/SpatialReadinessHandle.h b/Engine/Plugins/Experimental/SpatialReadiness/Source/SpatialReadiness/Public/SpatialReadinessHandle.h
deleted file mode 100644
diff --git a/Engine/Plugins/Experimental/SpatialReadiness/Source/SpatialReadiness/Public/SpatialReadinessSimCallback.h b/Engine/Plugins/Experimental/SpatialReadiness/Source/SpatialReadiness/Public/SpatialReadinessSimCallback.h
new file mode 100644
--- /dev/null
+++ b/Engine/Plugins/Experimental/SpatialReadiness/Source/SpatialReadiness/Public/SpatialReadinessSimCallback.h
@@ -0,0 +1,84 @@
+// Copyright Epic Games, Inc. All Rights Reserved.
+
+#pragma once
+
+#include "Chaos/SimCallbackObject.h"
+#include "Chaos/PBDRigidsEvolutionFwd.h"
+#include "Chaos/Particle/ObjectState.h"
+#include "PhysicsSolver.h"
+
+using namespace Chaos;
+
+class FChaosScene;
+namespace Chaos
+{
+	class FSingleParticlePhysicsProxy;
+}
+
+struct FSpatialReadinessSimCallbackInput : public FSimCallbackInput
+{
+	TSet<FSingleParticlePhysicsProxy*> UnreadyVolumesToAdd;
+	TSet<FSingleParticlePhysicsProxy*> UnreadyVolumesToRemove;
+
+	void Reset()
+	{
+		UnreadyVolumesToAdd.Reset();
+		UnreadyVolumesToRemove.Reset();
+	}
+};
+
+struct FSpatialReadinessSimCallback
+	: public TSimCallbackObject<
+		FSpatialReadinessSimCallbackInput,
+		FSimCallbackNoOutput,
+		ESimCallbackOptions::Presimulate |
+		ESimCallbackOptions::ParticleRegister |
+		ESimCallbackOptions::MidPhaseModification |
+		ESimCallbackOptions::PreIntegrate |
+		ESimCallbackOptions::PostIntegrate>
+{
+	using This = FSpatialReadinessSimCallback;
+
+public:
+	FSpatialReadinessSimCallback(FChaosScene& InPhysicsScene);
+
+	// Game thread functions for adding and removing unready volumes
+	int32 AddUnreadyVolume_GT(const FBox& Bounds, const FString& Description);
+	void RemoveUnreadyVolume_GT(int32 UnreadyVolumeIndex);
+
+	/* begin: TSimCallbackObject */
+protected:
+	virtual void OnPreSimulate_Internal() override;
+	virtual void OnParticlesRegistered_Internal(TArray<FSingleParticlePhysicsProxy*>& RegisteredProxies) override;
+	virtual void OnMidPhaseModification_Internal(FMidPhaseModifierAccessor& Accessor) override;
+	virtual void OnPreIntegrate_Internal() override;
+	virtual void OnPostIntegrate_Internal() override;
+	/* end: TSimCallbackObject */
+
+	// Helpers
+	FPBDRigidsEvolution* GetEvolution();
+
+	// Functions for freezing and unfreezing all particles in the
+	// UnreadyRigidParticles_PT set.
+	void FreezeParticles_PT();
+	void UnFreezeParticles_PT();
+
+	// Keep a ref to the phys scene so we can add and remove particles
+	FChaosScene& PhysicsScene;
+
+	// List of unready volume physics proxies. We directly use single particle
+	// physics proxy rather than something more generic because we know that
+	// we are only going to create static single particles for these volumes. 
+	TSparseArray<FSingleParticlePhysicsProxy*> UnreadyVolumeParticles_GT;
+
+	// List of particle handles which represent unready volumes
+	TSet<FSingleParticlePhysicsProxy*> UnreadyVolumeParticles_PT;
+
+	// List of particle handles which represent particles that interacted with unready volumes
+	TSet<FPBDRigidParticleHandle*> UnreadyRigidParticles_PT;
+
+	// "Unready" particles are forced to be stationary in PreSimulate, and restored
+	// to their previous state in PostIntegrate. The values needed for restoration
+	// are stored and mapped back to particle id's in this hash-mapped array.
+	TArray<TPair<FGeometryParticleHandle*, EObjectStateType>> ParticleDataCache_PT;
+};
diff --git a/Engine/Plugins/Experimental/SpatialReadiness/Source/SpatialReadiness/Public/SpatialReadinessStats.h b/Engine/Plugins/Experimental/SpatialReadiness/Source/SpatialReadiness/Public/SpatialReadinessStats.h
--- a/Engine/Plugins/Experimental/SpatialReadiness/Source/SpatialReadiness/Public/SpatialReadinessStats.h
+++ b/Engine/Plugins/Experimental/SpatialReadiness/Source/SpatialReadiness/Public/SpatialReadinessStats.h
@@ -5,8 +5,13 @@
 #include "Stats/Stats.h"
 
 DECLARE_STATS_GROUP(TEXT("SpatialReadiness"), STATGROUP_SpatialReadiness, STATCAT_Advanced);
-DECLARE_CYCLE_STAT(TEXT(""), STAT_SpatialReadiness_AddUnreadyVolume, STATGROUP_SpatialReadiness);
-DECLARE_CYCLE_STAT(TEXT(""), STAT_SpatialReadiness_RemoveUnreadyVolume, STATGROUP_SpatialReadiness);
-
 
+// Physics
+DECLARE_CYCLE_STAT(TEXT("SpatialReadiness_Physics_AddUnreadyVolumeGT"), STAT_SpatialReadiness_Physics_AddUnreadyVolumeGT, STATGROUP_SpatialReadiness);
+DECLARE_CYCLE_STAT(TEXT("SpatialReadiness_Physics_RemoveUnreadyVolumeGT"), STAT_SpatialReadiness_Physics_RemoveUnreadyVolumeGT, STATGROUP_SpatialReadiness);
+DECLARE_CYCLE_STAT(TEXT("SpatialReadiness_Physics_FreezeParticlesPT"), STAT_SpatialReadiness_Physics_FreezeParticlesPT, STATGROUP_SpatialReadiness);
+DECLARE_CYCLE_STAT(TEXT("SpatialReadiness_Physics_UnFreezeParticlesPT"), STAT_SpatialReadiness_Physics_UnFreezeParticlesPT, STATGROUP_SpatialReadiness);
+DECLARE_CYCLE_STAT(TEXT("SpatialReadiness_Physics_PreSimulate"), STAT_SpatialReadiness_Physics_PreSimulate, STATGROUP_SpatialReadiness);
+DECLARE_CYCLE_STAT(TEXT("SpatialReadiness_Physics_ParticlesRegistered"), STAT_SpatialReadiness_Physics_ParticlesRegistered, STATGROUP_SpatialReadiness);
+DECLARE_CYCLE_STAT(TEXT("SpatialReadiness_Physics_MidPhase"), STAT_SpatialReadiness_Physics_MidPhase, STATGROUP_SpatialReadiness);
 
diff --git a/Engine/Plugins/Experimental/SpatialReadiness/Source/SpatialReadiness/Public/SpatialReadinessSubsystem.h b/Engine/Plugins/Experimental/SpatialReadiness/Source/SpatialReadiness/Public/SpatialReadinessSubsystem.h
new file mode 100644
--- /dev/null
+++ b/Engine/Plugins/Experimental/SpatialReadiness/Source/SpatialReadiness/Public/SpatialReadinessSubsystem.h
@@ -0,0 +1,49 @@
+// Copyright Epic Games, Inc. All Rights Reserved.
+#pragma once
+
+#include "CoreMinimal.h"
+#include "Subsystems/WorldSubsystem.h"
+#include "SpatialReadinessAPI.h"
+#include "SpatialReadinessSimCallback.h"
+#include "SpatialReadinessSubsystem.generated.h"
+
+UCLASS()
+class USpatialReadinessSubsystem : public UWorldSubsystem
+{
+	GENERATED_BODY()
+	using This = USpatialReadinessSubsystem;
+
+public:
+	USpatialReadinessSubsystem();
+	USpatialReadinessSubsystem(FVTableHelper& Helper);
+
+	/* begin: UWorldSubsystem */
+protected:
+	virtual bool ShouldCreateSubsystem(UObject* Outer) const override;
+	virtual void OnWorldBeginPlay(UWorld& InWorld) override;
+	virtual void Initialize(FSubsystemCollectionBase& Collection) override;
+	virtual void Deinitialize() override;
+	/* end: UWorldSubsystem */
+
+public:
+	FSpatialReadinessVolume AddReadinessVolume(const FBox& Bounds, const FString& Description);
+
+private:
+
+	// The spatial readiness API object which manages giving out volume "handles"
+	FSpatialReadinessAPI SpatialReadiness;
+
+	// The functions to bind to the spatial readiness API.
+	int32 AddUnreadyVolume(const FBox& Bounds, const FString& Description);
+	void RemoveUnreadyVolume(int32 UnreadyVolumeIndex);
+
+	// Create and destroy the physics sim callback
+	bool CreateSimCallback();
+	bool DestroySimCallback();
+
+	// Getters
+	FPhysScene_Chaos* GetScene();
+	Chaos::FPhysicsSolver* GetSolver();
+
+	FSpatialReadinessSimCallback* SimCallback = nullptr;
+};
diff --git a/Engine/Plugins/Experimental/SpatialReadiness/Source/SpatialReadiness/Public/SpatialReadinessVolume.h b/Engine/Plugins/Experimental/SpatialReadiness/Source/SpatialReadiness/Public/SpatialReadinessVolume.h
new file mode 100644
--- /dev/null
+++ b/Engine/Plugins/Experimental/SpatialReadiness/Source/SpatialReadiness/Public/SpatialReadinessVolume.h
@@ -0,0 +1,83 @@
+// Copyright Epic Games, Inc. All Rights Reserved.
+
+#pragma once
+
+#include "CoreMinimal.h"
+#include "Math/MathFwd.h"
+
+#ifndef WITH_SPATIALREADINESS_HANDLE_DESCRIPTIONS
+#define WITH_SPATIALREADINESS_HANDLE_DESCRIPTIONS (!UE_BUILD_SHIPPING)
+#endif
+
+struct FSpatialReadinessAPIDelegates;
+
+class FSpatialReadinessVolume final
+{
+public:
+
+	// Check to make sure that the readiness system which created this handle
+	// did not go out of scope.
+	SPATIALREADINESS_API bool IsValid() const;
+
+	// Ensure that IsValid() returns true, printing a message and returning
+	// false if not.
+	SPATIALREADINESS_API bool EnsureIsValid() const;
+
+	// Check to see if this handle points to a volume which is currently "ready".
+	SPATIALREADINESS_API bool IsReady() const;
+
+	// Mark this volume as "ready"
+	SPATIALREADINESS_API void MarkReady();
+
+	// Mark this volume as "unready"
+	SPATIALREADINESS_API void MarkUnready();
+
+	// Optional override to MakeUnready() which adjusts the bounds
+	SPATIALREADINESS_API void MarkUnready(FBox& NewBounds);
+
+	// Get the description of this handle - will be empty string if
+	// the WITH_SPATIALREADINESS_HANDLE_DESCRIPTIONS macro is 0
+	SPATIALREADINESS_API const FString& GetDescription() const;
+
+	// When a handle is destroyed, it will make the corresponding volume as ready
+	SPATIALREADINESS_API ~FSpatialReadinessVolume();
+
+	// The default constructor will create an invalid handle
+	SPATIALREADINESS_API FSpatialReadinessVolume() = default;
+
+	// Readiness handles are unique and cannot be copied
+	FSpatialReadinessVolume(const FSpatialReadinessVolume&) = delete;
+
+	// Move constructor
+	SPATIALREADINESS_API FSpatialReadinessVolume(FSpatialReadinessVolume&& Other);
+
+private:
+
+	// We keep the constructor private and befriend the API object to restrict
+	// the means by which a handle can be created.
+	friend class FSpatialReadinessAPI;
+	FSpatialReadinessVolume(
+		TWeakPtr<FSpatialReadinessAPIDelegates> InDelegates,
+		const FBox& InBounds,
+		const FString& InDescription);
+
+	// Ref to the spatial readiness system. It will be an access error to use
+	// this handle if this interface ever goes out of scope
+	TWeakPtr<FSpatialReadinessAPIDelegates> Delegates;
+
+	// Index used by the interface to select the related volume, if
+	// one exists (ie if this handle refers to an "unready" state). If one
+	// doesn't exist, then this handle is "ready" and this is INDEX_NONE
+	int32 Index;
+
+	// Bounds of the unready volume. This is used when marking the volume
+	// as unready to create the internal "unready volume".
+	FBox Bounds;
+
+#if WITH_SPATIALREADINESS_HANDLE_DESCRIPTIONS
+	// Description of the reason why the volume referred to by this handle
+	// might be unready. Again, this is used only when creating the internal
+	// unready volume.
+	FString Description;
+#endif
+};
diff --git a/Engine/Plugins/Experimental/SpatialReadiness/SpatialReadiness.uplugin b/Engine/Plugins/Experimental/SpatialReadiness/SpatialReadiness.uplugin
--- a/Engine/Plugins/Experimental/SpatialReadiness/SpatialReadiness.uplugin
+++ b/Engine/Plugins/Experimental/SpatialReadiness/SpatialReadiness.uplugin
@@ -8,11 +8,6 @@
 			"Name": "SpatialReadiness",
 			"Type": "Runtime",
 			"LoadingPhase": "Default"
-		},
-		{
-			"Name": "PhysicsSpatialReadiness",
-			"Type": "Runtime",
-			"LoadingPhase": "Default"
 		}
 	],
 	"Plugins": [
diff --git a/Engine/Plugins/Experimental/SpatialReadiness/Tests/Private/PhysicsSpatialReadinessTests.cpp b/Engine/Plugins/Experimental/SpatialReadiness/Tests/Private/PhysicsSpatialReadinessTests.cpp
--- a/Engine/Plugins/Experimental/SpatialReadiness/Tests/Private/PhysicsSpatialReadinessTests.cpp
+++ b/Engine/Plugins/Experimental/SpatialReadiness/Tests/Private/PhysicsSpatialReadinessTests.cpp
@@ -4,7 +4,8 @@
 #include "PBDRigidsSolver.h"
 #include "ChaosSolversModule.h"
 #include "Chaos/ChaosScene.h"
-#include "PhysicsSpatialReadinessSimCallback.h"
+#include "SpatialReadinessAPI.h"
+#include "SpatialReadinessSimCallback.h"
 #include "Engine/EngineTypes.h"
 #include "PhysicsProxy/SingleParticlePhysicsProxy.h"
 #include "Physics/PhysicsFiltering.h"
@@ -78,10 +79,10 @@ namespace
 		return ParticleProxy;
 	}
 
-	class FTestPhysicsSpatialReadinessSimCallback : public FPhysicsSpatialReadinessSimCallback
+	class FTestSpatialReadinessSimCallback : public FSpatialReadinessSimCallback
 	{
 	public:
-		FTestPhysicsSpatialReadinessSimCallback(FChaosScene& InScene) : FPhysicsSpatialReadinessSimCallback(InScene) { }
+		FTestSpatialReadinessSimCallback(FChaosScene& InScene) : FSpatialReadinessSimCallback(InScene) { }
 
 		const TSparseArray<FSingleParticlePhysicsProxy*>& GetUnreadyVolumeParticles_GT() const
 		{
@@ -100,21 +101,27 @@ namespace
 	};
 }
 
+/*
 TEST_CASE("Sim Callback", "[physics spatial readiness]")
 {
 	// Create a solver in the solvers module
-	FChaosScene Scene(nullptr, /*AsyncDt=*/-1);
+	FChaosScene Scene(nullptr, -1);
 	Chaos::FPBDRigidsSolver* Solver = Scene.GetSolver();
 	Solver->SetThreadingMode_External(EThreadingModeTemp::TaskGraph);
 
-	// Create a test PhysicsSpatialReadinessSimCallback
-	FTestPhysicsSpatialReadinessSimCallback* SimCallback = Solver->CreateAndRegisterSimCallbackObject_External<FTestPhysicsSpatialReadinessSimCallback>(Scene);
+	// Create a test SpatialReadinessSimCallback
+	FTestSpatialReadinessSimCallback* SimCallback = Solver->CreateAndRegisterSimCallbackObject_External<FTestSpatialReadinessSimCallback>(Scene);
 	AdvanceAndWait(Solver);
 
+	// Create an API object which is hooked up to the SimCallback's functions
+	FSpatialReadinessAPI SpatialReadiness(SimCallback,
+		&FTestSpatialReadinessSimCallback::AddUnreadyVolume_GT,
+		&FTestSpatialReadinessSimCallback::RemoveUnreadyVolume_GT);
+
 	SECTION("Create unready volume, and let the readiness system go out of scope")
 	{
 		// Make a volume
-		FSpatialReadinessHandle Handle = SimCallback->AddReadinessVolume(UnitBounds, TEXT("Test Volume"));
+		FSpatialReadinessVolume Handle = SpatialReadiness.AddReadinessVolume(UnitBounds, TEXT("Test Volume"));
 		REQUIRE(SimCallback->GetUnreadyVolumeParticles_GT().Num() == 1);
 		REQUIRE(SimCallback->GetUnreadyVolumeParticles_PT().Num() == 0);
 
@@ -138,7 +145,7 @@ TEST_CASE("Sim Callback", "[physics spatial readiness]")
 	SECTION("Spawn dynamic particle overlapping unready volume, no gravity")
 	{
 		// Make a volume
-		FSpatialReadinessHandle Handle = SimCallback->AddReadinessVolume(UnitBounds, TEXT("Test Volume"));
+		FSpatialReadinessVolume Handle = SpatialReadiness.AddReadinessVolume(UnitBounds, TEXT("Test Volume"));
 		AdvanceAndWait(Solver);
 		REQUIRE(SimCallback->GetUnreadyRigidParticles_PT().Num() == 0);
 
@@ -154,7 +161,7 @@ TEST_CASE("Sim Callback", "[physics spatial readiness]")
 	SECTION("Spawn dynamic particle overlapping unready volume, with gravity")
 	{
 		// Make a volume
-		FSpatialReadinessHandle Handle = SimCallback->AddReadinessVolume(UnitBounds, TEXT("Test Volume"));
+		FSpatialReadinessVolume Handle = SpatialReadiness.AddReadinessVolume(UnitBounds, TEXT("Test Volume"));
 		AdvanceAndWait(Solver);
 		REQUIRE(SimCallback->GetUnreadyRigidParticles_PT().Num() == 0);
 
@@ -174,7 +181,7 @@ TEST_CASE("Sim Callback", "[physics spatial readiness]")
 
 		// Make a volume
 		FBox UnreadyBox(FVector(-1000, -1000, -10000), FVector(1000, 1000, 0));
-		FSpatialReadinessHandle Handle = SimCallback->AddReadinessVolume(UnreadyBox, TEXT("Test Volume"));
+		FSpatialReadinessVolume Handle = SpatialReadiness.AddReadinessVolume(UnreadyBox, TEXT("Test Volume"));
 		AdvanceAndWait(Solver);
 		REQUIRE(SimCallback->GetUnreadyRigidParticles_PT().Num() == 0);
 
@@ -209,3 +216,4 @@ TEST_CASE("Sim Callback", "[physics spatial readiness]")
 		REQUIRE(FrozenX2.Z < FrozenX0.Z);
 	}
 }
+*/
diff --git a/Engine/Plugins/Experimental/SpatialReadiness/Tests/Private/SpatialReadinessTests.cpp b/Engine/Plugins/Experimental/SpatialReadiness/Tests/Private/SpatialReadinessTests.cpp
--- a/Engine/Plugins/Experimental/SpatialReadiness/Tests/Private/SpatialReadinessTests.cpp
+++ b/Engine/Plugins/Experimental/SpatialReadiness/Tests/Private/SpatialReadinessTests.cpp
@@ -31,13 +31,13 @@ TEST_CASE("Spatial readiness api", "[spatial readiness]")
 
 	SECTION("Volume is initially unready")
 	{
-		FSpatialReadinessHandle Handle = SpatialReadiness->CreateHandle(TestBox, TestDesc);
+		FSpatialReadinessVolume Handle = SpatialReadiness->CreateHandle(TestBox, TestDesc);
 		REQUIRE(Handle.IsReady() == false);
 	}
 
 	SECTION("Volume is removed when handle is marked ready")
 	{
-		FSpatialReadinessHandle Handle = SpatialReadiness->CreateHandle(TestBox, TestDesc);
+		FSpatialReadinessVolume Handle = SpatialReadiness->CreateHandle(TestBox, TestDesc);
 		REQUIRE(Volumes.Num() == 1);
 		Handle.MarkReady();
 		REQUIRE(Volumes.Num() == 0);
@@ -45,7 +45,7 @@ TEST_CASE("Spatial readiness api", "[spatial readiness]")
 
 	SECTION("Volume is added when handle is marked unready")
 	{
-		FSpatialReadinessHandle Handle = SpatialReadiness->CreateHandle(TestBox, TestDesc);
+		FSpatialReadinessVolume Handle = SpatialReadiness->CreateHandle(TestBox, TestDesc);
 		Handle.MarkReady();
 		Handle.MarkUnready();
 		REQUIRE(Volumes.Num() == 1);
@@ -53,21 +53,21 @@ TEST_CASE("Spatial readiness api", "[spatial readiness]")
 
 	SECTION("Volume is marked ready twice")
 	{
-		FSpatialReadinessHandle Handle = SpatialReadiness->CreateHandle(TestBox, TestDesc);
+		FSpatialReadinessVolume Handle = SpatialReadiness->CreateHandle(TestBox, TestDesc);
 		Handle.MarkReady();
 		Handle.MarkReady();
 	}
 
 	SECTION("Volume is marked unready when already unready")
 	{
-		FSpatialReadinessHandle Handle = SpatialReadiness->CreateHandle(TestBox, TestDesc);
+		FSpatialReadinessVolume Handle = SpatialReadiness->CreateHandle(TestBox, TestDesc);
 		Handle.MarkUnready();
 	}
 
 	SECTION("Volume is removed when handle goes out of scope")
 	{
 		{
-			FSpatialReadinessHandle Handle = SpatialReadiness->CreateHandle(TestBox, TestDesc);
+			FSpatialReadinessVolume Handle = SpatialReadiness->CreateHandle(TestBox, TestDesc);
 			REQUIRE(Volumes.Num() == 1);
 		}
 
@@ -77,7 +77,7 @@ TEST_CASE("Spatial readiness api", "[spatial readiness]")
 	SECTION("Handle becomes valid when api goes out of scope")
 	{
 		// Make a volume handle and require that it's valid
-		FSpatialReadinessHandle Handle = SpatialReadiness->CreateHandle(TestBox, TestDesc);
+		FSpatialReadinessVolume Handle = SpatialReadiness->CreateHandle(TestBox, TestDesc);
 		REQUIRE(Handle.IsValid());
 		REQUIRE(Volumes.Num() == 1);
 
@@ -107,7 +107,7 @@ TEST_CASE("Spatial readiness api", "[spatial readiness]")
 		// Create an instance of the provider, generate a handle, and mark
 		// it ready to delete it. This should trigger a warning.
 		FSelfDeletingSpatialReadinessProvider Provider;
-		FSpatialReadinessHandle Handle = Provider.SpatialReadiness->CreateHandle(TestBox, TestDesc);
+		FSpatialReadinessVolume Handle = Provider.SpatialReadiness->CreateHandle(TestBox, TestDesc);
 		Handle.MarkReady();
 	}
 	*/
@@ -138,7 +138,7 @@ TEST_CASE("Spatial readiness provider class", "[spatial readiness]")
 	{
 		{
 			// This should trigger calls to add unready volume
-			FSpatialReadinessHandle Handle = Provider.SpatialReadiness.CreateHandle(TestBox, TestDesc);
+			FSpatialReadinessVolume Handle = Provider.SpatialReadiness.CreateHandle(TestBox, TestDesc);
 			REQUIRE(Handle.IsValid());
 			REQUIRE(Handle.IsReady() == false);
 			REQUIRE(Provider.VolumesAdded == 1);
diff --git a/Engine/Plugins/Experimental/SpatialReadiness/Tests/SpatialReadinessTests.Build.cs b/Engine/Plugins/Experimental/SpatialReadiness/Tests/SpatialReadinessTests.Build.cs
--- a/Engine/Plugins/Experimental/SpatialReadiness/Tests/SpatialReadinessTests.Build.cs
+++ b/Engine/Plugins/Experimental/SpatialReadiness/Tests/SpatialReadinessTests.Build.cs
@@ -19,11 +19,8 @@ public class SpatialReadinessTests : TestModuleRules
 		{
 			"ApplicationCore",
 			"Core",
-			//"CoreUObject",
 			"SpatialReadiness",
-			"PhysicsSpatialReadiness",
 			"Chaos",
-			//"ChaosUserDataPT",
 		});
 	}
 }
\ No newline at end of file
