HORDE: Add CoreUObject and Chaos to SpatialReadiness.Build.cs

diff --git a/Engine/Plugins/Experimental/SpatialReadiness/Source/SpatialReadiness/SpatialReadiness.Build.cs b/Engine/Plugins/Experimental/SpatialReadiness/Source/SpatialReadiness/SpatialReadiness.Build.cs
--- a/Engine/Plugins/Experimental/SpatialReadiness/Source/SpatialReadiness/SpatialReadiness.Build.cs
+++ b/Engine/Plugins/Experimental/SpatialReadiness/Source/SpatialReadiness/SpatialReadiness.Build.cs
@@ -10,6 +10,7 @@ public class SpatialReadiness : ModuleRules
             new string[]
             {
                 "Engine",
+				"CoreUObject",
             });
 
         PublicDependencyModuleNames.AddRange(
@@ -17,5 +18,7 @@ public class SpatialReadiness : ModuleRules
                 "Core",
             }
         );
+
+		SetupModulePhysicsSupport(Target);
     }
 }
