Add readiness query functions on GT and PT. Also rename CreateHandle to CreateVolume to match new naming.

Changed to using HashMappedArray for faster lookups in queries. Had to add a Remove function to HashMappedArray.

diff --git a/Engine/Plugins/Experimental/SpatialReadiness/Source/SpatialReadiness/Private/SpatialReadiness/SpatialReadinessAPI.cpp b/Engine/Plugins/Experimental/SpatialReadiness/Source/SpatialReadiness/Private/SpatialReadiness/SpatialReadinessAPI.cpp
--- a/Engine/Plugins/Experimental/SpatialReadiness/Source/SpatialReadiness/Private/SpatialReadiness/SpatialReadinessAPI.cpp
+++ b/Engine/Plugins/Experimental/SpatialReadiness/Source/SpatialReadiness/Private/SpatialReadiness/SpatialReadinessAPI.cpp
@@ -21,7 +21,7 @@ FSpatialReadinessAPI::~FSpatialReadinessAPI()
 	}
 }
 
-FSpatialReadinessVolume FSpatialReadinessAPI::CreateHandle(const FBox& Bounds, const FString& Description)
+FSpatialReadinessVolume FSpatialReadinessAPI::CreateVolume(const FBox& Bounds, const FString& Description)
 {
 	// Create a handle for a new unready volume
 	FSpatialReadinessVolume Handle(Delegates.ToWeakPtr(), Bounds, Description);
diff --git a/Engine/Plugins/Experimental/SpatialReadiness/Source/SpatialReadiness/Private/SpatialReadiness/SpatialReadinessSimCallback.cpp b/Engine/Plugins/Experimental/SpatialReadiness/Source/SpatialReadiness/Private/SpatialReadiness/SpatialReadinessSimCallback.cpp
--- a/Engine/Plugins/Experimental/SpatialReadiness/Source/SpatialReadiness/Private/SpatialReadiness/SpatialReadinessSimCallback.cpp
+++ b/Engine/Plugins/Experimental/SpatialReadiness/Source/SpatialReadiness/Private/SpatialReadiness/SpatialReadinessSimCallback.cpp
@@ -1,22 +1,26 @@
 // Copyright Epic Games, Inc. All Rights Reserved.
 
 #include "SpatialReadinessSimCallback.h"
-#include "PhysicsProxy/SingleParticlePhysicsProxy.h"
+#include "Chaos/Box.h"
 #include "Chaos/ChaosEngineInterface.h"
-#include "Chaos/ChaosScene.h"
 #include "Chaos/Collision/SimSweep.h"
 #include "Chaos/MidPhaseModification.h"
 #include "Chaos/ParticleHandle.h"
+#include "Chaos/PBDRigidsEvolutionGBF.h"
 #include "SpatialReadinessVolume.h"
 #include "PBDRigidsSolver.h"
-#include "Chaos/PBDRigidsEvolutionGBF.h"
+#include "PhysicsProxy/SingleParticlePhysicsProxy.h"
 #include "Physics/PhysicsFiltering.h"
+#include "Physics/PhysicsInterfaceUtils.h"
+#include "Physics/Experimental/ChaosInterfaceWrapper.h"
+#include "Physics/Experimental/PhysScene_Chaos.h"
 #include "SpatialReadinessStats.h"
 
 using namespace Chaos;
 
-FSpatialReadinessSimCallback::FSpatialReadinessSimCallback(FChaosScene& InPhysicsScene)
+FSpatialReadinessSimCallback::FSpatialReadinessSimCallback(FPhysScene_Chaos& InPhysicsScene)
 	: PhysicsScene(InPhysicsScene)
+	, UnreadyVolumeData_GT(256)
 	, ParticleDataCache_PT()
 { }
 
@@ -45,7 +49,7 @@ int32 FSpatialReadinessSimCallback::AddUnreadyVolume_GT(const FBox& Bounds, cons
 
 	// Create a box implicit geometry from the bounds
 	const FVec3 BoxCenter = (Bounds.Min + Bounds.Max) * .5f;
-	const FVec3 BoxHalfExtent = Bounds.Max - Bounds.Min;
+	const FVec3 BoxHalfExtent = (Bounds.Max - Bounds.Min) * .5f;
 	Chaos::FImplicitObjectPtr BoxGeom = MakeImplicitObjectPtr<TBox<FReal, 3>>(-BoxHalfExtent, BoxHalfExtent);
 
 	// Create a new static particle to represent the volume
@@ -84,9 +88,10 @@ int32 FSpatialReadinessSimCallback::AddUnreadyVolume_GT(const FBox& Bounds, cons
 	// Make the geometry
 	ParticleHandle.SetGeometry(BoxGeom);
 	ParticleHandle.SetShapeSimCollisionEnabled(0, true);
-	ParticleHandle.SetShapeQueryCollisionEnabled(0, false);
+	ParticleHandle.SetShapeQueryCollisionEnabled(0, true);
 	ParticleHandle.SetShapeSimData(0, SimFilterData);
 	ParticleHandle.ShapesArray()[0]->SetIsProbe(true);
+	ParticleHandle.ShapesArray()[0]->SetQueryData(QueryFilterData);
 #if CHAOS_DEBUG_NAME
 	ParticleHandle.SetDebugName(MakeShared<FString>(TEXT("Unready Volume")));
 #endif
@@ -95,34 +100,46 @@ int32 FSpatialReadinessSimCallback::AddUnreadyVolume_GT(const FBox& Bounds, cons
 	TArray<FPhysicsActorHandle> Actors = { ParticleProxy };
 	PhysicsScene.AddActorsToScene_AssumesLocked(Actors);
 
-	// Save the proxy in our list of GT particles
-	const int32 ParticleIndex = UnreadyVolumeParticles_GT.Add(ParticleProxy);
+	// Save the proxy in our list of GT particles.
+	//
+	// Technically, the index of the volume just has to be unique, it doesn't
+	// have to be the same as the particle index. However, we do this
+	// just so that when querying there's no need to map from hit particles
+	// back to volume index - we'll already have direct access to the index.
+	const int32 ParticleIndex = ParticleHandle.UniqueIdx().Idx;
+	const int32 VolumeIndex = ParticleIndex;
+	ensureMsgf(UnreadyVolumeData_GT.TryAdd(VolumeIndex,
+		FUnreadyVolumeData_GT(
+			ParticleProxy,
+			Description
+		)), TEXT("Failed to add volume data to map - VolumeIndex already exists!"));
 
 	// Queue up the particle proxy for processing on PT
 	SimInput.UnreadyVolumesToAdd.Add(ParticleProxy);
 	SimInput.UnreadyVolumesToRemove.Remove(ParticleProxy);
 
-	return ParticleIndex;
+	return VolumeIndex;
 }
 
 void FSpatialReadinessSimCallback::RemoveUnreadyVolume_GT(int32 UnreadyVolumeIndex)
 {
 	SCOPE_CYCLE_COUNTER(STAT_SpatialReadiness_Physics_RemoveUnreadyVolumeGT)
 
-	if (!ensureMsgf(UnreadyVolumeParticles_GT.IsValidIndex(UnreadyVolumeIndex), TEXT("Trying to remove unready volume whose index is not being tracked")))
+	FUnreadyVolumeData_GT* VolumeData = UnreadyVolumeData_GT.Find(UnreadyVolumeIndex);
+	if (!ensureMsgf(VolumeData, TEXT("Trying to remove unready volume whos index is not being tracked")))
 	{
 		return;
 	}
 
 	// Get the proxy associated with this index
-	FSingleParticlePhysicsProxy* ParticleProxy = UnreadyVolumeParticles_GT[UnreadyVolumeIndex];
+	FSingleParticlePhysicsProxy* ParticleProxy = VolumeData->Proxy;
 	if (!ensureMsgf(ParticleProxy, TEXT("Particle proxy associated with unready volume index was null")))
 	{
 		return;
 	}
 	
 	// Free the index in our GT tracker
-	UnreadyVolumeParticles_GT.RemoveAt(UnreadyVolumeIndex);
+	UnreadyVolumeData_GT.Remove(UnreadyVolumeIndex);
 
 	// Tell the PT to remove it's tracking of this proxy as well
 	FSpatialReadinessSimCallbackInput* SimInput = GetProducerInputData_External();
@@ -136,6 +153,159 @@ void FSpatialReadinessSimCallback::RemoveUnreadyVolume_GT(int32 UnreadyVolumeInd
 	FChaosEngineInterface::ReleaseActor(ParticleProxy, &PhysicsScene);
 }
 
+bool FSpatialReadinessSimCallback::QueryReadiness_GT(const FBox& Bounds, TArray<int32>& OutVolumeIndices, bool bAllUnreadyVolumes) const
+{
+	// Constants that we'll use to set up query parameters
+	constexpr static uint8 Channel = static_cast<uint8>(ECollisionChannel::ECC_WorldStatic);
+	constexpr static bool bComplex = false;
+	const bool bMulti = bAllUnreadyVolumes;
+
+	// Query objects
+	FCollisionQueryParams QueryParams(SCENE_QUERY_STAT(ReadinessQuery), false);
+	FCollisionObjectQueryParams ObjectParams(Channel);
+	FCollisionResponseContainer ResponseContainer;
+	FCollisionFilterData CollisionFilterData = CreateQueryFilterData(Channel, bComplex, ResponseContainer, QueryParams, ObjectParams, bMulti);
+	FOverlapAllQueryCallback QueryCallback;
+	EQueryFlags QueryFlags = EQueryFlags::PreFilter;// | EQueryFlags::SkipNarrowPhase;
+	ChaosInterface::FQueryFilterData QueryFilterData = ChaosInterface::MakeQueryFilterData(CollisionFilterData, QueryFlags, QueryParams);
+	ChaosInterface::FQueryDebugParams DebugParams = { };
+
+	// Create the geometry needed for the query
+	Chaos::FImplicitBox3 Geom = Chaos::FImplicitBox3(Bounds.Min, Bounds.Max, 0);
+
+	// Do the actual query
+	FPhysicsHitCallback<ChaosInterface::FOverlapHit> HitBuffer;
+	FPhysicsCommand::ExecuteRead(&PhysicsScene,
+	[this, &Geom, &HitBuffer, &QueryFlags, &CollisionFilterData, &QueryFilterData, &QueryCallback, &DebugParams]()
+	{
+		LowLevelOverlap(PhysicsScene, Geom, FTransform::Identity, HitBuffer, QueryFlags, CollisionFilterData, QueryFilterData, &QueryCallback, DebugParams);
+	});
+
+	// Make sure the hits from the buffer are actually in our list of unready volumes,
+	// and collect their indices
+	OutVolumeIndices.Reset();
+	for (int32 HitIdx = 0; HitIdx < HitBuffer.GetNumHits(); ++HitIdx)
+	{
+		// Get the hit result from this index
+		const ChaosInterface::FOverlapHit& Hit = HitBuffer.GetHits()[HitIdx];
+
+		// Get the particle that we hit
+		Chaos::FGeometryParticle* HitParticle = Hit.Actor;
+		if (HitParticle == nullptr)
+		{
+			continue;
+		}
+
+		// Get the particle index
+		//
+		// NOTE: Since we're just directly using particle unique indices for
+		// volume indices, this mapping is simplified. We may at some point
+		// want to use a more complex mapping though, in which case we'll need
+		// to do something different here.
+		const uint32 VolumeIndex = HitParticle->UniqueIdx().Idx;
+
+		// Make sure that we have an actual entry for this particle in our
+		// unready volumes list
+		//if (!UnreadyVolumeData_GT.IsValidIndex(VolumeIndex))
+		if (UnreadyVolumeData_GT.Find(VolumeIndex) == nullptr)
+		{
+			continue;
+		}
+
+		// Add the index to the output list
+		OutVolumeIndices.Add(VolumeIndex);
+	}
+
+	// If we didn't hit any unready volumes then that means this volume is "ready"
+	return OutVolumeIndices.Num() == 0;
+}
+
+const FUnreadyVolumeData_GT* FSpatialReadinessSimCallback::GetVolumeData_GT(int32 VolumeIndex) const
+{
+	return UnreadyVolumeData_GT.Find(VolumeIndex);
+}
+
+bool FSpatialReadinessSimCallback::QueryReadiness_PT(const FAABB3& Bounds, TArray<const FSingleParticlePhysicsProxy*>& OutVolumeProxies)
+{
+	// Get the evolution, and get the acceleration structure from it
+	FPBDRigidsEvolution* Evolution = GetEvolution();
+	if (Evolution == nullptr)
+	{
+		return false;
+	}
+
+	// Set up particle filters for the query, so that we only get unready volumes
+	// NOTE: If we could speed up this bit here, that'd be great
+	const auto& ParticleFilter = [this](const FGeometryParticleHandle* Particle) -> bool
+	{
+		// If the particle has a single particle physics proxy (and it does,
+		// I can almost guarantee it), then count it as a hit if it's in our
+		// list of unready volume particles.
+		if (const IPhysicsProxyBase* Proxy = Particle->PhysicsProxy())
+		{
+			if (Proxy->GetType() == EPhysicsProxyType::SingleParticleProxy)
+			{
+				return UnreadyVolumeParticles_PT.Contains(
+					static_cast<const Chaos::FSingleParticlePhysicsProxy*>(Proxy));
+			}
+		}
+		return false;
+	};
+
+	// We know that unready volumes will all be boxes, so in theory we could
+	// probably filter a bunch of interactions before getting to the particle
+	// filter, however the particle filter is applied first so there's no point
+	// in implementing a shape filter.
+	const auto& ShapeFilter = [](const FPerShapeData* Shape, const FImplicitObject* Implicit)
+	{
+		return true;
+	};
+
+	// Make a lambda for collecting 
+	TArray<Private::FSimOverlapParticleShape> Overlaps;
+	const auto& OverlapCollector = [&Overlaps](const Private::FSimOverlapParticleShape& Overlap)
+	{
+		Overlaps.Add(Overlap);
+	};
+
+	// Get the broadphase and query against it to see if this new rigid particle
+	// generates midphases with any unready volumes.
+	//
+	// NOTE: In theory we could have a first-hit version of this function which
+	// would potentially avoid some unnecessary checking.
+	//
+	// NOTE: The above note is only valid until we actually start registering
+	// freeze locks with unready volumes...
+	Private::SimOverlapBounds(Evolution->GetSpatialAcceleration(), Bounds, ParticleFilter, ShapeFilter, OverlapCollector);
+
+	// Convert overlaps to list of hit proxies
+	OutVolumeProxies.Reset();
+	for (const Private::FSimOverlapParticleShape& Overlap : Overlaps)
+	{
+		const Chaos::FGeometryParticleHandle* ParticleHandle = Overlap.HitParticle;
+		if (ParticleHandle == nullptr)
+		{
+			continue;
+		}
+
+		const IPhysicsProxyBase* Proxy = ParticleHandle->PhysicsProxy();
+		if (Proxy == nullptr)
+		{
+			continue;
+		}
+
+		if (Proxy->GetType() != EPhysicsProxyType::SingleParticleProxy)
+		{
+			continue;
+		}
+
+		OutVolumeProxies.Add(static_cast<const Chaos::FSingleParticlePhysicsProxy*>(Proxy));
+	}
+
+	// Return true if there were NO overlaps - meaning, the volume is ready
+	return OutVolumeProxies.Num() == 0;
+}
+
 void FSpatialReadinessSimCallback::FreezeParticles_PT()
 {
 	SCOPE_CYCLE_COUNTER(STAT_SpatialReadiness_Physics_FreezeParticlesPT)
@@ -217,13 +387,6 @@ void FSpatialReadinessSimCallback::OnParticlesRegistered_Internal(TArray<FSingle
 {
 	SCOPE_CYCLE_COUNTER(STAT_SpatialReadiness_Physics_ParticlesRegistered)
 
-	// Get the evolution, and get the acceleration structure from it
-	FPBDRigidsEvolution* Evolution = GetEvolution();
-	if (Evolution == nullptr)
-	{
-		return;
-	}
-
 	// For each newly added particle, query against unready volumes to see if
 	// it should be frozen.
 	for (FSingleParticlePhysicsProxy* ParticleProxy : RegisteredProxies)
@@ -241,54 +404,12 @@ void FSpatialReadinessSimCallback::OnParticlesRegistered_Internal(TArray<FSingle
 		}
 
 		// Get the bounds to use for the query
+		TArray<const FSingleParticlePhysicsProxy*> VolumeProxies;
 		const FAABB3 QueryBounds = RigidParticle->WorldSpaceInflatedBounds();
-
-		// Set up particle filters for the query, so that we only get unready volumes
-		// NOTE: If we could speed up this bit here, that'd be great
-		const auto& ParticleFilter = [this](const FGeometryParticleHandle* Particle) -> bool
-		{
-			// If the particle has a single particle physics proxy (and it does,
-			// I can almost guarantee it), then count it as a hit if it's in our
-			// list of unready volume particles.
-			const IPhysicsProxyBase* Proxy = Particle->PhysicsProxy();
-			if (Proxy->GetType() == EPhysicsProxyType::SingleParticleProxy)
-			{
-				return UnreadyVolumeParticles_PT.Contains(
-					static_cast<const Chaos::FSingleParticlePhysicsProxy*>(Proxy));
-			}
-			return false;
-		};
-
-		// We know that unready volumes will all be boxes, so in theory we could
-		// probably filter a bunch of interactions before getting to the particle
-		// filter, however the particle filter is applied first so there's no point
-		// in implementing a shape filter.
-		const auto& ShapeFilter = [](const FPerShapeData* Shape, const FImplicitObject* Implicit)
-		{
-			return true;
-		};
-
-		// Make a lambda for collecting 
-		TArray<Private::FSimOverlapParticleShape> Overlaps;
-		const auto& OverlapCollector = [&Overlaps](const Private::FSimOverlapParticleShape& Overlap)
-		{
-			Overlaps.Add(Overlap);
-		};
-
-		// Get the broadphase and query against it to see if this new rigid particle
-		// generates midphases with any unready volumes.
-		//
-		// NOTE: In theory we could have a first-hit version of this function which
-		// would potentially avoid some unnecessary checking.
-		//
-		// NOTE: The above note is only valid until we actually start registering
-		// freeze locks with unready volumes...
-		Private::SimOverlapBounds(Evolution->GetSpatialAcceleration(), QueryBounds, ParticleFilter, ShapeFilter, OverlapCollector);
-
-		// The overlaps have already been filtered for unready volume interaction,
-		// so if there were _any_ overlaps, all we need to do is freeze the particle
-		if (Overlaps.Num() > 0)
+		if (!QueryReadiness_PT(QueryBounds, VolumeProxies))
 		{
+			// If this volume hit an unready volume, query readiness returns false
+			// and we must mark the particle as frozen
 			UnreadyRigidParticles_PT.Add(RigidParticle);
 		}
 	}
@@ -349,3 +470,13 @@ void FSpatialReadinessSimCallback::OnPostIntegrate_Internal()
 {
 	UnFreezeParticles_PT();
 }
+
+uint32 FSpatialReadinessSimCallback::FHashMapTraits::GetElementID(const FUnreadyVolumeData_GT& Element)
+{
+	if (FSingleParticlePhysicsProxy* Proxy = Element.Proxy)
+	{
+		return Proxy->GetGameThreadAPI().UniqueIdx().Idx;
+	}
+	return INDEX_NONE;
+}
+
diff --git a/Engine/Plugins/Experimental/SpatialReadiness/Source/SpatialReadiness/Private/SpatialReadiness/SpatialReadinessSubsystem.cpp b/Engine/Plugins/Experimental/SpatialReadiness/Source/SpatialReadiness/Private/SpatialReadiness/SpatialReadinessSubsystem.cpp
--- a/Engine/Plugins/Experimental/SpatialReadiness/Source/SpatialReadiness/Private/SpatialReadiness/SpatialReadinessSubsystem.cpp
+++ b/Engine/Plugins/Experimental/SpatialReadiness/Source/SpatialReadiness/Private/SpatialReadiness/SpatialReadinessSubsystem.cpp
@@ -7,15 +7,15 @@
 
 using namespace Chaos;
 
-USpatialReadinessSubsystem::USpatialReadinessSubsystem()
+USpatialReadiness::USpatialReadiness()
 	: Super()
 	, SpatialReadiness(this, &This::AddUnreadyVolume, &This::RemoveUnreadyVolume) { }
 
-USpatialReadinessSubsystem::USpatialReadinessSubsystem(FVTableHelper& Helper)
+USpatialReadiness::USpatialReadiness(FVTableHelper& Helper)
 	: Super(Helper)
 	, SpatialReadiness(this, &This::AddUnreadyVolume, &This::RemoveUnreadyVolume) { }
 
-bool USpatialReadinessSubsystem::ShouldCreateSubsystem(UObject* Outer) const
+bool USpatialReadiness::ShouldCreateSubsystem(UObject* Outer) const
 {
 	// TODO: Where should the setting for this exist? Is it enough
 	//       to just not load the physics readiness module if we don't
@@ -23,32 +23,60 @@ bool USpatialReadinessSubsystem::ShouldCreateSubsystem(UObject* Outer) const
 	return true;
 }
 
-void USpatialReadinessSubsystem::OnWorldBeginPlay(UWorld& InWorld)
+void USpatialReadiness::OnWorldBeginPlay(UWorld& InWorld)
 {
 	Super::OnWorldBeginPlay(InWorld);
 	CreateSimCallback();
 }
 
-void USpatialReadinessSubsystem::Initialize(FSubsystemCollectionBase& Collection)
+void USpatialReadiness::Initialize(FSubsystemCollectionBase& Collection)
 {
 	Super::Initialize(Collection);
 }
 
-void USpatialReadinessSubsystem::Deinitialize()
+void USpatialReadiness::Deinitialize()
 {
 	DestroySimCallback();
 	Super::Deinitialize();
 }
 
-FSpatialReadinessVolume USpatialReadinessSubsystem::AddReadinessVolume(const FBox& Bounds, const FString& Description)
+FSpatialReadinessVolume USpatialReadiness::AddReadinessVolume(const FBox& Bounds, const FString& Description)
 {
 	// Create a readiness volume and return it's "handle".
 	// This call will trigger the associated AddUnreadyVolume method
 	// since volumes are unready by default.
-	return SpatialReadiness.CreateHandle(Bounds, Description);
+	return SpatialReadiness.CreateVolume(Bounds, Description);
 }
 
-int32 USpatialReadinessSubsystem::AddUnreadyVolume(const FBox& Bounds, const FString& Description)
+bool USpatialReadiness::QueryReadiness(const FBox& Bounds, TArray<FString>& OutDescriptions, const bool bAllDescriptions) const
+{
+	// Add the physics unready volume
+	if (ensureMsgf(SimCallback, TEXT("Tried to query for readiness when no sim callback exists")))
+	{
+		// Query for readiness volumes in the sim callback object which tracks them
+		TArray<int32> VolumeIndices;
+		const bool bReady = SimCallback->QueryReadiness_GT(Bounds, VolumeIndices, bAllDescriptions);
+
+		// If the macro is set to build with descriptions, populate the output
+		// descriptions array
+#if WITH_SPATIAL_READINESS_DESCRIPTIONS
+		if (bReady == false)
+		{
+			OutDescriptions.Reset(VolumeIndices.Num());
+			for (int32 VolumeIndex : VolumeIndices)
+			{
+				OutDescriptions.Add(SimCallback->GetVolumeData_GT(VolumeIndex)->Description);
+			}
+		}
+#endif
+		return bReady;
+	}
+
+	// Default to not-ready if there was a problem
+	return false;
+}
+
+int32 USpatialReadiness::AddUnreadyVolume(const FBox& Bounds, const FString& Description)
 {
 	// Add the physics unready volume
 	if (ensureMsgf(SimCallback, TEXT("Tried to add unready volume when no sim callback exists")))
@@ -59,7 +87,7 @@ int32 USpatialReadinessSubsystem::AddUnreadyVolume(const FBox& Bounds, const FSt
 	return INDEX_NONE;
 }
 
-void USpatialReadinessSubsystem::RemoveUnreadyVolume(int32 UnreadyVolumeIndex)
+void USpatialReadiness::RemoveUnreadyVolume(int32 UnreadyVolumeIndex)
 {
 	// Remove the physics unready volume
 	if (ensureMsgf(SimCallback, TEXT("Tried to add unready volume when no sim callback exists")))
@@ -68,7 +96,7 @@ void USpatialReadinessSubsystem::RemoveUnreadyVolume(int32 UnreadyVolumeIndex)
 	}
 }
 
-bool USpatialReadinessSubsystem::CreateSimCallback()
+bool USpatialReadiness::CreateSimCallback()
 {
 	// If we already have a sim callback, destroy it.
 	if (SimCallback)
@@ -99,7 +127,7 @@ bool USpatialReadinessSubsystem::CreateSimCallback()
 	}
 
 	// Request creation of the scene callback from the solver
-	SimCallback = Solver->CreateAndRegisterSimCallbackObject_External<FSpatialReadinessSimCallback>(static_cast<FChaosScene&>(*Scene));
+	SimCallback = Solver->CreateAndRegisterSimCallbackObject_External<FSpatialReadinessSimCallback>(*Scene);
 	if (!ensureMsgf(SimCallback, TEXT("Sim callback creation failed")))
 	{
 		return false;
@@ -109,7 +137,7 @@ bool USpatialReadinessSubsystem::CreateSimCallback()
 	return true;
 }
 
-bool USpatialReadinessSubsystem::DestroySimCallback()
+bool USpatialReadiness::DestroySimCallback()
 {
 	if (SimCallback)
 	{
@@ -124,7 +152,7 @@ bool USpatialReadinessSubsystem::DestroySimCallback()
 	return false;
 }
 
-FPhysScene_Chaos* USpatialReadinessSubsystem::GetScene()
+FPhysScene_Chaos* USpatialReadiness::GetScene()
 {
 	if (UWorld* World = GetWorld())
 	{
@@ -136,7 +164,7 @@ FPhysScene_Chaos* USpatialReadinessSubsystem::GetScene()
 	return nullptr;
 }
 
-Chaos::FPhysicsSolver* USpatialReadinessSubsystem::GetSolver()
+Chaos::FPhysicsSolver* USpatialReadiness::GetSolver()
 {
 	if (FPhysScene* Scene = GetScene())
 	{
diff --git a/Engine/Plugins/Experimental/SpatialReadiness/Source/SpatialReadiness/Private/SpatialReadiness/SpatialReadinessVolume.cpp b/Engine/Plugins/Experimental/SpatialReadiness/Source/SpatialReadiness/Private/SpatialReadiness/SpatialReadinessVolume.cpp
--- a/Engine/Plugins/Experimental/SpatialReadiness/Source/SpatialReadiness/Private/SpatialReadiness/SpatialReadinessVolume.cpp
+++ b/Engine/Plugins/Experimental/SpatialReadiness/Source/SpatialReadiness/Private/SpatialReadiness/SpatialReadinessVolume.cpp
@@ -10,7 +10,7 @@ FSpatialReadinessVolume::FSpatialReadinessVolume(
 	: Delegates(InDelegates)
 	, Index(INDEX_NONE)
 	, Bounds(InBounds)
-#if WITH_SPATIALREADINESS_HANDLE_DESCRIPTIONS
+#if WITH_SPATIAL_READINESS_DESCRIPTIONS
 	, Description(InDescription)
 #endif
 { }
@@ -19,14 +19,14 @@ FSpatialReadinessVolume::FSpatialReadinessVolume(FSpatialReadinessVolume&& Other
 	: Delegates(Other.Delegates)
 	, Index(Other.Index)
 	, Bounds(Other.Bounds)
-#if WITH_SPATIALREADINESS_HANDLE_DESCRIPTIONS
+#if WITH_SPATIAL_READINESS_DESCRIPTIONS
 	, Description(Other.Description)
 #endif
 {
 	Other.Delegates.Reset();
 	Other.Index = INDEX_NONE;
 	Other.Bounds = FBox();
-#if WITH_SPATIALREADINESS_HANDLE_DESCRIPTIONS
+#if WITH_SPATIAL_READINESS_DESCRIPTIONS
 	Other.Description = TEXT("");
 #endif
 }
@@ -127,7 +127,7 @@ void FSpatialReadinessVolume::MarkUnready(FBox& NewBounds)
 
 const FString& FSpatialReadinessVolume::GetDescription() const
 {
-#if WITH_SPATIALREADINESS_HANDLE_DESCRIPTIONS
+#if WITH_SPATIAL_READINESS_DESCRIPTIONS
 	return Description;
 #else
 	static const FString EmptyString;
diff --git a/Engine/Plugins/Experimental/SpatialReadiness/Source/SpatialReadiness/Public/SpatialReadinessAPI.h b/Engine/Plugins/Experimental/SpatialReadiness/Source/SpatialReadiness/Public/SpatialReadinessAPI.h
--- a/Engine/Plugins/Experimental/SpatialReadiness/Source/SpatialReadiness/Public/SpatialReadinessAPI.h
+++ b/Engine/Plugins/Experimental/SpatialReadiness/Source/SpatialReadiness/Public/SpatialReadinessAPI.h
@@ -29,9 +29,10 @@ public:
 
 	SPATIALREADINESS_API ~FSpatialReadinessAPI();
 
-	SPATIALREADINESS_API FSpatialReadinessVolume CreateHandle(const FBox& Bounds, const FString& Description);
+	SPATIALREADINESS_API FSpatialReadinessVolume CreateVolume(const FBox& Bounds, const FString& Description);
 
 private:
+
 	TSharedPtr<FSpatialReadinessAPIDelegates> Delegates;
 };
 
diff --git a/Engine/Plugins/Experimental/SpatialReadiness/Source/SpatialReadiness/Public/SpatialReadinessSimCallback.h b/Engine/Plugins/Experimental/SpatialReadiness/Source/SpatialReadiness/Public/SpatialReadinessSimCallback.h
--- a/Engine/Plugins/Experimental/SpatialReadiness/Source/SpatialReadiness/Public/SpatialReadinessSimCallback.h
+++ b/Engine/Plugins/Experimental/SpatialReadiness/Source/SpatialReadiness/Public/SpatialReadinessSimCallback.h
@@ -5,16 +5,32 @@
 #include "Chaos/SimCallbackObject.h"
 #include "Chaos/PBDRigidsEvolutionFwd.h"
 #include "Chaos/Particle/ObjectState.h"
+#include "Chaos/Framework/HashMappedArray.h"
+#include "SpatialReadinessVolume.h"
 #include "PhysicsSolver.h"
 
 using namespace Chaos;
 
-class FChaosScene;
+class FPhysScene_Chaos;
 namespace Chaos
 {
 	class FSingleParticlePhysicsProxy;
 }
 
+struct FUnreadyVolumeData_GT
+{
+	FUnreadyVolumeData_GT(FSingleParticlePhysicsProxy* InProxy, const FString& InDescription)
+		: Proxy(InProxy)
+#if WITH_SPATIAL_READINESS_DESCRIPTIONS
+		, Description(InDescription)
+#endif
+	{ }
+	FSingleParticlePhysicsProxy* Proxy;
+#if WITH_SPATIAL_READINESS_DESCRIPTIONS
+	FString Description;
+#endif
+};
+
 struct FSpatialReadinessSimCallbackInput : public FSimCallbackInput
 {
 	TSet<FSingleParticlePhysicsProxy*> UnreadyVolumesToAdd;
@@ -40,12 +56,24 @@ struct FSpatialReadinessSimCallback
 	using This = FSpatialReadinessSimCallback;
 
 public:
-	FSpatialReadinessSimCallback(FChaosScene& InPhysicsScene);
+	FSpatialReadinessSimCallback(FPhysScene_Chaos& InPhysicsScene);
 
 	// Game thread functions for adding and removing unready volumes
 	int32 AddUnreadyVolume_GT(const FBox& Bounds, const FString& Description);
 	void RemoveUnreadyVolume_GT(int32 UnreadyVolumeIndex);
 
+	// Game thread function for querying for unready volumes
+	//
+	// If bAllUnreadyVolumes is true, then a multi-query will be used and
+	// OutVolumeIndices will be populated with every index that encroaches.
+	//
+	// For performance, we only return the index of the first unread volume
+	// that we find.
+	bool QueryReadiness_GT(const FBox& Bounds, TArray<int32>& OutVolumeIndices, bool bAllUnreadyVolumes=false) const;
+
+	// Given a volume index, get it's description
+	const FUnreadyVolumeData_GT* GetVolumeData_GT(int32 VolumeIndex) const;
+
 	/* begin: TSimCallbackObject */
 protected:
 	virtual void OnPreSimulate_Internal() override;
@@ -63,13 +91,22 @@ protected:
 	void FreezeParticles_PT();
 	void UnFreezeParticles_PT();
 
+	// Physics thread function for querying for unready volumes
+	// TODO: Make const
+	bool QueryReadiness_PT(const FAABB3& Bounds, TArray<const FSingleParticlePhysicsProxy*>& OutVolumeProxies);
+
 	// Keep a ref to the phys scene so we can add and remove particles
-	FChaosScene& PhysicsScene;
+	FPhysScene_Chaos& PhysicsScene;
 
 	// List of unready volume physics proxies. We directly use single particle
 	// physics proxy rather than something more generic because we know that
 	// we are only going to create static single particles for these volumes. 
-	TSparseArray<FSingleParticlePhysicsProxy*> UnreadyVolumeParticles_GT;
+	struct FHashMapTraits
+	{
+		static uint32 GetIDHash(const int32 Idx) { return MurmurFinalize32(Idx); }
+		static uint32 GetElementID(const FUnreadyVolumeData_GT& Element);
+	};
+	Private::THashMappedArray<int32, FUnreadyVolumeData_GT, FHashMapTraits> UnreadyVolumeData_GT;
 
 	// List of particle handles which represent unready volumes
 	TSet<FSingleParticlePhysicsProxy*> UnreadyVolumeParticles_PT;
diff --git a/Engine/Plugins/Experimental/SpatialReadiness/Source/SpatialReadiness/Public/SpatialReadinessSubsystem.h b/Engine/Plugins/Experimental/SpatialReadiness/Source/SpatialReadiness/Public/SpatialReadinessSubsystem.h
--- a/Engine/Plugins/Experimental/SpatialReadiness/Source/SpatialReadiness/Public/SpatialReadinessSubsystem.h
+++ b/Engine/Plugins/Experimental/SpatialReadiness/Source/SpatialReadiness/Public/SpatialReadinessSubsystem.h
@@ -8,14 +8,24 @@
 #include "SpatialReadinessSubsystem.generated.h"
 
 UCLASS()
-class USpatialReadinessSubsystem : public UWorldSubsystem
+class USpatialReadiness : public UWorldSubsystem
 {
 	GENERATED_BODY()
-	using This = USpatialReadinessSubsystem;
+	using This = USpatialReadiness;
 
 public:
-	USpatialReadinessSubsystem();
-	USpatialReadinessSubsystem(FVTableHelper& Helper);
+	USpatialReadiness();
+	USpatialReadiness(FVTableHelper& Helper);
+
+	// Add a volume which can be marked read/unready
+	FSpatialReadinessVolume AddReadinessVolume(const FBox& Bounds, const FString& Description);
+
+	// Check to see if a volume is ready
+	//
+	// If bAllDescriptions = false, OutDescriptions will contain only the description of the
+	// first unready volume that we hit. If true, it will contain descriptions from all
+	// overlapping volumes.
+	bool QueryReadiness(const FBox& Bounds, TArray<FString>& OutDescriptions, bool bAllDescriptions = false) const;
 
 	/* begin: UWorldSubsystem */
 protected:
@@ -25,9 +35,6 @@ protected:
 	virtual void Deinitialize() override;
 	/* end: UWorldSubsystem */
 
-public:
-	FSpatialReadinessVolume AddReadinessVolume(const FBox& Bounds, const FString& Description);
-
 private:
 
 	// The spatial readiness API object which manages giving out volume "handles"
diff --git a/Engine/Plugins/Experimental/SpatialReadiness/Source/SpatialReadiness/Public/SpatialReadinessVolume.h b/Engine/Plugins/Experimental/SpatialReadiness/Source/SpatialReadiness/Public/SpatialReadinessVolume.h
--- a/Engine/Plugins/Experimental/SpatialReadiness/Source/SpatialReadiness/Public/SpatialReadinessVolume.h
+++ b/Engine/Plugins/Experimental/SpatialReadiness/Source/SpatialReadiness/Public/SpatialReadinessVolume.h
@@ -5,8 +5,8 @@
 #include "CoreMinimal.h"
 #include "Math/MathFwd.h"
 
-#ifndef WITH_SPATIALREADINESS_HANDLE_DESCRIPTIONS
-#define WITH_SPATIALREADINESS_HANDLE_DESCRIPTIONS (!UE_BUILD_SHIPPING)
+#ifndef WITH_SPATIAL_READINESS_DESCRIPTIONS
+#define WITH_SPATIAL_READINESS_DESCRIPTIONS (!UE_BUILD_SHIPPING)
 #endif
 
 struct FSpatialReadinessAPIDelegates;
@@ -36,7 +36,7 @@ public:
 	SPATIALREADINESS_API void MarkUnready(FBox& NewBounds);
 
 	// Get the description of this handle - will be empty string if
-	// the WITH_SPATIALREADINESS_HANDLE_DESCRIPTIONS macro is 0
+	// the WITH_SPATIAL_READINESS_DESCRIPTIONS macro is 0
 	SPATIALREADINESS_API const FString& GetDescription() const;
 
 	// When a handle is destroyed, it will make the corresponding volume as ready
@@ -74,7 +74,7 @@ private:
 	// as unready to create the internal "unready volume".
 	FBox Bounds;
 
-#if WITH_SPATIALREADINESS_HANDLE_DESCRIPTIONS
+#if WITH_SPATIAL_READINESS_DESCRIPTIONS
 	// Description of the reason why the volume referred to by this handle
 	// might be unready. Again, this is used only when creating the internal
 	// unready volume.
diff --git a/Engine/Plugins/Experimental/SpatialReadiness/Tests/Private/PhysicsSpatialReadinessTests.cpp b/Engine/Plugins/Experimental/SpatialReadiness/Tests/Private/PhysicsSpatialReadinessTests.cpp
--- a/Engine/Plugins/Experimental/SpatialReadiness/Tests/Private/PhysicsSpatialReadinessTests.cpp
+++ b/Engine/Plugins/Experimental/SpatialReadiness/Tests/Private/PhysicsSpatialReadinessTests.cpp
@@ -3,7 +3,7 @@
 #include "CoreMinimal.h"
 #include "PBDRigidsSolver.h"
 #include "ChaosSolversModule.h"
-#include "Chaos/ChaosScene.h"
+#include "Physics/Experimental/PhysScene_Chaos.h"
 #include "SpatialReadinessAPI.h"
 #include "SpatialReadinessSimCallback.h"
 #include "Engine/EngineTypes.h"
@@ -82,12 +82,7 @@ namespace
 	class FTestSpatialReadinessSimCallback : public FSpatialReadinessSimCallback
 	{
 	public:
-		FTestSpatialReadinessSimCallback(FChaosScene& InScene) : FSpatialReadinessSimCallback(InScene) { }
-
-		const TSparseArray<FSingleParticlePhysicsProxy*>& GetUnreadyVolumeParticles_GT() const
-		{
-			return UnreadyVolumeParticles_GT;
-		}
+		FTestSpatialReadinessSimCallback(FPhysScene_Chaos& InScene) : FSpatialReadinessSimCallback(InScene) { }
 
 		const TSet<FSingleParticlePhysicsProxy*>& GetUnreadyVolumeParticles_PT() const
 		{
@@ -101,11 +96,10 @@ namespace
 	};
 }
 
-/*
 TEST_CASE("Sim Callback", "[physics spatial readiness]")
 {
 	// Create a solver in the solvers module
-	FChaosScene Scene(nullptr, -1);
+	FPhysScene_Chaos Scene;
 	Chaos::FPBDRigidsSolver* Solver = Scene.GetSolver();
 	Solver->SetThreadingMode_External(EThreadingModeTemp::TaskGraph);
 
@@ -114,38 +108,15 @@ TEST_CASE("Sim Callback", "[physics spatial readiness]")
 	AdvanceAndWait(Solver);
 
 	// Create an API object which is hooked up to the SimCallback's functions
-	FSpatialReadinessAPI SpatialReadiness(SimCallback,
+	FSpatialReadinessAPI SpatialReadiness(
+		static_cast<FSpatialReadinessSimCallback*>(SimCallback),
 		&FTestSpatialReadinessSimCallback::AddUnreadyVolume_GT,
 		&FTestSpatialReadinessSimCallback::RemoveUnreadyVolume_GT);
 
-	SECTION("Create unready volume, and let the readiness system go out of scope")
-	{
-		// Make a volume
-		FSpatialReadinessVolume Handle = SpatialReadiness.AddReadinessVolume(UnitBounds, TEXT("Test Volume"));
-		REQUIRE(SimCallback->GetUnreadyVolumeParticles_GT().Num() == 1);
-		REQUIRE(SimCallback->GetUnreadyVolumeParticles_PT().Num() == 0);
-
-		// Advance the solver, which should result in the unready volume existing
-		// on the physics thread
-		AdvanceAndWait(Solver);
-		REQUIRE(SimCallback->GetUnreadyVolumeParticles_GT().Num() == 1);
-		REQUIRE(SimCallback->GetUnreadyVolumeParticles_PT().Num() == 1);
-
-		// The volume handle should be "valid"
-		REQUIRE(Handle.IsValid());
-
-		// Destroy the sim callback
-		Solver->UnregisterAndFreeSimCallbackObject_External(SimCallback);
-		AdvanceAndWait(Solver);
-
-		// The volume handle should now be "invalid"
-		REQUIRE(!Handle.IsValid());
-	}
-
 	SECTION("Spawn dynamic particle overlapping unready volume, no gravity")
 	{
 		// Make a volume
-		FSpatialReadinessVolume Handle = SpatialReadiness.AddReadinessVolume(UnitBounds, TEXT("Test Volume"));
+		FSpatialReadinessVolume Volume = SpatialReadiness.CreateVolume(UnitBounds, TEXT("Test Volume"));
 		AdvanceAndWait(Solver);
 		REQUIRE(SimCallback->GetUnreadyRigidParticles_PT().Num() == 0);
 
@@ -161,7 +132,7 @@ TEST_CASE("Sim Callback", "[physics spatial readiness]")
 	SECTION("Spawn dynamic particle overlapping unready volume, with gravity")
 	{
 		// Make a volume
-		FSpatialReadinessVolume Handle = SpatialReadiness.AddReadinessVolume(UnitBounds, TEXT("Test Volume"));
+		FSpatialReadinessVolume Volume = SpatialReadiness.CreateVolume(UnitBounds, TEXT("Test Volume"));
 		AdvanceAndWait(Solver);
 		REQUIRE(SimCallback->GetUnreadyRigidParticles_PT().Num() == 0);
 
@@ -181,7 +152,7 @@ TEST_CASE("Sim Callback", "[physics spatial readiness]")
 
 		// Make a volume
 		FBox UnreadyBox(FVector(-1000, -1000, -10000), FVector(1000, 1000, 0));
-		FSpatialReadinessVolume Handle = SpatialReadiness.AddReadinessVolume(UnreadyBox, TEXT("Test Volume"));
+		FSpatialReadinessVolume Volume = SpatialReadiness.CreateVolume(UnreadyBox, TEXT("Test Volume"));
 		AdvanceAndWait(Solver);
 		REQUIRE(SimCallback->GetUnreadyRigidParticles_PT().Num() == 0);
 
@@ -203,9 +174,9 @@ TEST_CASE("Sim Callback", "[physics spatial readiness]")
 			REQUIRE(FrozenX0 == FrozenX1);
 		}
 
-		// Mark the volume handle as "ready" and advance again - the particle
+		// Mark the volume as "ready" and advance again - the particle
 		// should be removed from the list of unready particles
-		Handle.MarkReady();
+		Volume.MarkReady();
 		AdvanceAndWait(Solver);
 		REQUIRE(SimCallback->GetUnreadyRigidParticles_PT().Num() == 0);
 
@@ -215,5 +186,26 @@ TEST_CASE("Sim Callback", "[physics spatial readiness]")
 		const FVec3 FrozenX2 = BoxProxy->GetPhysicsThreadAPI()->X();
 		REQUIRE(FrozenX2.Z < FrozenX0.Z);
 	}
+
+	SECTION("Game thread query for readiness")
+	{
+		// Make a volume
+		FBox UnreadyBox(FVector(-100), FVector(100));
+		FSpatialReadinessVolume Volume = SpatialReadiness.CreateVolume(UnreadyBox, TEXT("Test Volume"));
+
+		// Make temp vars
+		TArray<int32> VolumeIndices;
+		bool bIsReady;
+
+		// Do a query which should intersect the unready area
+		bIsReady = SimCallback->QueryReadiness_GT(FBox(FVector(0), FVector(100)), VolumeIndices);
+		REQUIRE(!bIsReady);
+		REQUIRE(VolumeIndices.Num() == 1);
+		REQUIRE(VolumeIndices[0] == 0);
+
+		// Do a query which should not intersect the unready area
+		bIsReady = SimCallback->QueryReadiness_GT(FBox(FVector(200), FVector(300)), VolumeIndices);
+		REQUIRE(bIsReady);
+		REQUIRE(VolumeIndices.Num() == 0);
+	}
 }
-*/
diff --git a/Engine/Plugins/Experimental/SpatialReadiness/Tests/Private/SpatialReadinessTests.cpp b/Engine/Plugins/Experimental/SpatialReadiness/Tests/Private/SpatialReadinessTests.cpp
--- a/Engine/Plugins/Experimental/SpatialReadiness/Tests/Private/SpatialReadinessTests.cpp
+++ b/Engine/Plugins/Experimental/SpatialReadiness/Tests/Private/SpatialReadinessTests.cpp
@@ -31,59 +31,59 @@ TEST_CASE("Spatial readiness api", "[spatial readiness]")
 
 	SECTION("Volume is initially unready")
 	{
-		FSpatialReadinessVolume Handle = SpatialReadiness->CreateHandle(TestBox, TestDesc);
-		REQUIRE(Handle.IsReady() == false);
+		FSpatialReadinessVolume Volume = SpatialReadiness->CreateVolume(TestBox, TestDesc);
+		REQUIRE(Volume.IsReady() == false);
 	}
 
-	SECTION("Volume is removed when handle is marked ready")
+	SECTION("Volume is removed when volume is marked ready")
 	{
-		FSpatialReadinessVolume Handle = SpatialReadiness->CreateHandle(TestBox, TestDesc);
+		FSpatialReadinessVolume Volume = SpatialReadiness->CreateVolume(TestBox, TestDesc);
 		REQUIRE(Volumes.Num() == 1);
-		Handle.MarkReady();
+		Volume.MarkReady();
 		REQUIRE(Volumes.Num() == 0);
 	}
 
-	SECTION("Volume is added when handle is marked unready")
+	SECTION("Volume is added when volume is marked unready")
 	{
-		FSpatialReadinessVolume Handle = SpatialReadiness->CreateHandle(TestBox, TestDesc);
-		Handle.MarkReady();
-		Handle.MarkUnready();
+		FSpatialReadinessVolume Volume = SpatialReadiness->CreateVolume(TestBox, TestDesc);
+		Volume.MarkReady();
+		Volume.MarkUnready();
 		REQUIRE(Volumes.Num() == 1);
 	}
 
 	SECTION("Volume is marked ready twice")
 	{
-		FSpatialReadinessVolume Handle = SpatialReadiness->CreateHandle(TestBox, TestDesc);
-		Handle.MarkReady();
-		Handle.MarkReady();
+		FSpatialReadinessVolume Volume = SpatialReadiness->CreateVolume(TestBox, TestDesc);
+		Volume.MarkReady();
+		Volume.MarkReady();
 	}
 
 	SECTION("Volume is marked unready when already unready")
 	{
-		FSpatialReadinessVolume Handle = SpatialReadiness->CreateHandle(TestBox, TestDesc);
-		Handle.MarkUnready();
+		FSpatialReadinessVolume Volume = SpatialReadiness->CreateVolume(TestBox, TestDesc);
+		Volume.MarkUnready();
 	}
 
-	SECTION("Volume is removed when handle goes out of scope")
+	SECTION("Volume is removed when volume goes out of scope")
 	{
 		{
-			FSpatialReadinessVolume Handle = SpatialReadiness->CreateHandle(TestBox, TestDesc);
+			FSpatialReadinessVolume Volume = SpatialReadiness->CreateVolume(TestBox, TestDesc);
 			REQUIRE(Volumes.Num() == 1);
 		}
 
 		REQUIRE(Volumes.Num() == 0);
 	}
 
-	SECTION("Handle becomes valid when api goes out of scope")
+	SECTION("Volume becomes valid when api goes out of scope")
 	{
-		// Make a volume handle and require that it's valid
-		FSpatialReadinessVolume Handle = SpatialReadiness->CreateHandle(TestBox, TestDesc);
-		REQUIRE(Handle.IsValid());
+		// Make a volume volume and require that it's valid
+		FSpatialReadinessVolume Volume = SpatialReadiness->CreateVolume(TestBox, TestDesc);
+		REQUIRE(Volume.IsValid());
 		REQUIRE(Volumes.Num() == 1);
 
-		// Delete the readiness api and ensure that the handle is now invalid
+		// Delete the readiness api and ensure that the volume is now invalid
 		SpatialReadiness.Reset();
-		REQUIRE(!Handle.IsValid());
+		REQUIRE(!Volume.IsValid());
 	}
 
 	/*
@@ -138,14 +138,14 @@ TEST_CASE("Spatial readiness provider class", "[spatial readiness]")
 	{
 		{
 			// This should trigger calls to add unready volume
-			FSpatialReadinessVolume Handle = Provider.SpatialReadiness.CreateHandle(TestBox, TestDesc);
-			REQUIRE(Handle.IsValid());
-			REQUIRE(Handle.IsReady() == false);
+			FSpatialReadinessVolume Volume = Provider.SpatialReadiness.CreateVolume(TestBox, TestDesc);
+			REQUIRE(Volume.IsValid());
+			REQUIRE(Volume.IsReady() == false);
 			REQUIRE(Provider.VolumesAdded == 1);
 			REQUIRE(Provider.VolumesRemoved == 0);
 		}
 
-		// The handle going out of scope should cause the volume to be removed
+		// The volume going out of scope should cause the volume to be removed
 		REQUIRE(Provider.VolumesAdded == 1);
 		REQUIRE(Provider.VolumesRemoved == 1);
 	}
diff --git a/Engine/Source/Runtime/Experimental/Chaos/Private/Chaos/PBDRigidsEvolutionGBF.cpp b/Engine/Source/Runtime/Experimental/Chaos/Private/Chaos/PBDRigidsEvolutionGBF.cpp
--- a/Engine/Source/Runtime/Experimental/Chaos/Private/Chaos/PBDRigidsEvolutionGBF.cpp
+++ b/Engine/Source/Runtime/Experimental/Chaos/Private/Chaos/PBDRigidsEvolutionGBF.cpp
@@ -1623,7 +1623,7 @@ TArray<FGeometryParticleHandle*> FPBDRigidsEvolutionGBF::GetConnectedParticles(F
 	struct FHashMapTraits
 	{
 		static uint32 GetIDHash(const FGeometryParticleHandle* Particle) { return MurmurFinalize32(uint32(Particle->UniqueIdx().Idx)); }
-		static bool ElementHasID(const FGeometryParticleHandle* A, const FGeometryParticleHandle* B) { return A == B; }
+		static const FGeometryParticleHandle* GetElementID(const FGeometryParticleHandle* Element) { return Element; }
 	};
 	Private::THashMappedArray<FGeometryParticleHandle*, FGeometryParticleHandle*, FHashMapTraits> ConnectedParticles(256);
 
diff --git a/Engine/Source/Runtime/Experimental/Chaos/Public/Chaos/Collision/MeshContactGenerator.h b/Engine/Source/Runtime/Experimental/Chaos/Public/Chaos/Collision/MeshContactGenerator.h
--- a/Engine/Source/Runtime/Experimental/Chaos/Public/Chaos/Collision/MeshContactGenerator.h
+++ b/Engine/Source/Runtime/Experimental/Chaos/Public/Chaos/Collision/MeshContactGenerator.h
@@ -495,7 +495,7 @@ namespace Chaos::Private
 		struct FEdgeTriangleIndicesMapTraits
 		{
 			static uint32 GetIDHash(const FContactEdgeID& EdgeID) { return uint32(MurmurFinalize64(EdgeID.EdgeID)); }
-			static bool ElementHasID(const FEdgeTriangleIndices& TriangleIndices, const FContactEdgeID& EdgeID) { return TriangleIndices.ID == EdgeID; }
+			static FContactEdgeID GetElementID(const FEdgeTriangleIndices& TriangleIndices) { return TriangleIndices.ID; }
 		};
 		THashMappedArray<FContactEdgeID, FEdgeTriangleIndices, FEdgeTriangleIndicesMapTraits> EdgeTriangleIndicesMap;
 
@@ -503,7 +503,7 @@ namespace Chaos::Private
 		struct VertexContactIndicesMapTraits
 		{
 			static uint32 GetIDHash(const FContactVertexID& VertexID) { return MurmurFinalize32(VertexID); }
-			static bool ElementHasID(const FVertexContactIndex& VertexIndex, const FContactVertexID& VertexID) { return VertexIndex.ID == VertexID; }
+			static FContactVertexID GetElementID(const FVertexContactIndex& VertexIndex) { return VertexIndex.ID; }
 		};
 		THashMappedArray<FContactVertexID, FVertexContactIndex, VertexContactIndicesMapTraits> VertexContactIndicesMap;
 	};
diff --git a/Engine/Source/Runtime/Experimental/Chaos/Public/Chaos/Framework/HashMappedArray.h b/Engine/Source/Runtime/Experimental/Chaos/Public/Chaos/Framework/HashMappedArray.h
--- a/Engine/Source/Runtime/Experimental/Chaos/Public/Chaos/Framework/HashMappedArray.h
+++ b/Engine/Source/Runtime/Experimental/Chaos/Public/Chaos/Framework/HashMappedArray.h
@@ -19,10 +19,9 @@ namespace Chaos::Private
 			return MurmurFinalize32(ID);
 		}
 
-		// Return true if the element is the one with the specified ID
-		static bool ElemenatHasID(const FElementType& Element, const FIDType& ID)
+		static FIDType GetElementID(const FElementType& Element)
 		{
-			return Element == ID;
+			return FIDType(Element);
 		}
 	};
 
@@ -120,6 +119,45 @@ namespace Chaos::Private
 			HashTable.Add(Key, Index);
 		}
 
+		void Remove(const FIDType ID)
+		{
+			const FHashType Key = FTraits::GetIDHash(ID);
+			for (uint32 Index = HashTable.First(Key); HashTable.IsValid(Index); Index = HashTable.Next(Index))
+			{
+				if (FTraits::GetElementID(Elements[Index]) == ID)
+				{
+					const uint32 LastElemIndex = Elements.Num() - 1;
+
+					// Remove the key/index from the hash table
+					HashTable.Remove(Key, Index);
+
+					// If this is the last element in the array (or if there's only one) then
+					// no need for fancy swaps.
+					if (Index == LastElemIndex)
+					{
+						Elements.RemoveAt(Index);
+					}
+
+					// If this is not the last element in the array, swap with the last element
+					// and fix up the hash table
+					else
+					{
+						const FIDType LastElemID = FTraits::GetElementID(Elements[LastElemIndex]);
+						const FHashType LastElemKey = FTraits::GetIDHash(LastElemID);
+
+						// Remove the last element's key/index from the hash table
+						HashTable.Remove(LastElemKey, LastElemID);
+
+						// Remove the element we want to remove, swapping it with the last element
+						Elements.RemoveAtSwap(Index);
+
+						// Re-add the last element's key, reusing the index of the element we just removed
+						HashTable.Add(LastElemKey, Index);
+					}
+				}
+			}
+		}
+
 		// Find the element with the specified ID. Roughly O(Max(1,N/M)) for N elements with a hash table of size M
 		const FElementType* Find(const FIDType ID) const
 		{
@@ -132,7 +170,7 @@ namespace Chaos::Private
 			const FHashType Key = FTraits::GetIDHash(ID);
 			for (uint32 Index = HashTable.First(Key); HashTable.IsValid(Index); Index = HashTable.Next(Index))
 			{
-				if (FTraits::ElementHasID(Elements[Index], ID))
+				if (FTraits::GetElementID(Elements[Index]) == ID)
 				{
 					return &Elements[Index];
 				}
