Add spatial readiness debug map in physics sandbox project. Add spatial readiness show-debug draw command.

#virtualized

diff --git a/Engine/Plugins/Experimental/SpatialReadiness/Config/DefaultInput.ini b/Engine/Plugins/Experimental/SpatialReadiness/Config/DefaultInput.ini
new file mode 100644
--- /dev/null
+++ b/Engine/Plugins/Experimental/SpatialReadiness/Config/DefaultInput.ini
@@ -0,0 +1,2 @@
+[/Script/EngineSettings.ConsoleSettings]
++ManualAutoCompleteList=(Command="ShowDebug",Desc="SpatialReadiness")
\ No newline at end of file
diff --git a/Engine/Plugins/Experimental/SpatialReadiness/Source/SpatialReadiness/Private/SpatialReadiness/SpatialReadinessSimCallback.cpp b/Engine/Plugins/Experimental/SpatialReadiness/Source/SpatialReadiness/Private/SpatialReadiness/SpatialReadinessSimCallback.cpp
--- a/Engine/Plugins/Experimental/SpatialReadiness/Source/SpatialReadiness/Private/SpatialReadiness/SpatialReadinessSimCallback.cpp
+++ b/Engine/Plugins/Experimental/SpatialReadiness/Source/SpatialReadiness/Private/SpatialReadiness/SpatialReadinessSimCallback.cpp
@@ -181,16 +181,6 @@ bool FSpatialReadinessSimCallback::QueryReadiness_GT(const FBox& Bounds, TArray<
 	// Create the geometry needed for the query
 	Chaos::FImplicitBox3 Geom = Chaos::FImplicitBox3(Bounds.Min, Bounds.Max, 0);
 
-	/*
-	// Do the actual query
-	TArray<FOverlapResult> Overlaps;
-	FPhysicsCommand::ExecuteRead(&PhysicsScene,
-	[this, &Geom, &Overlaps, &QueryParams, &ObjectParams]()
-	{
-		Chaos::Private::FGenericPhysicsSceneInterface::GeomOverlapMulti<FPhysicsGeometry>(PhysicsScene, Geom, FVec3(0), FQuat::Identity, Overlaps, Channel, QueryParams, FCollisionResponseParams::DefaultResponseParam, ObjectParams);
-	});
-	*/
-
 	// Do the query
 	FPhysicsHitCallback<ChaosInterface::FOverlapHit> HitBuffer;
 	FPhysicsCommand::ExecuteRead(&PhysicsScene,
@@ -236,49 +226,6 @@ bool FSpatialReadinessSimCallback::QueryReadiness_GT(const FBox& Bounds, TArray<
 
 	// If we didn't hit any unready volumes then that means this volume is "ready"
 	return OutVolumeIndices.Num() == 0;
-
-	/*
-	// Make sure the hits from the buffer are actually in our list of unready volumes,
-	// and collect their indices
-	OutVolumeIndices.Reset();
-	for (const FOverlapResult& Overlap : Overlaps)
-	{
-		Chaos::FPhysicsObjectHandle Object = Overlap.PhysicsObject;
-		if (Object == nullptr)
-		{
-			continue;
-		}
-
-		FLockedReadPhysicsObjectExternalInterface Interface = FPhysicsObjectExternalInterface::LockRead(Object);
-		const Chaos::FGeometryParticle* HitParticle = Interface->GetParticle(Object);
-		if (HitParticle == nullptr)
-		{
-			continue;
-		}
-
-		// Get the particle index
-		//
-		// NOTE: Since we're just directly using particle unique indices for
-		// volume indices, this mapping is simplified. We may at some point
-		// want to use a more complex mapping though, in which case we'll need
-		// to do something different here.
-		const uint32 VolumeIndex = HitParticle->UniqueIdx().Idx;
-
-		// Make sure that we have an actual entry for this particle in our
-		// unready volumes list
-		//if (!UnreadyVolumeData_GT.IsValidIndex(VolumeIndex))
-		if (UnreadyVolumeData_GT.Find(VolumeIndex) == nullptr)
-		{
-			continue;
-		}
-
-		// Add the index to the output list
-		OutVolumeIndices.Add(VolumeIndex);
-	}
-
-	// If we didn't hit any unready volumes then that means this volume is "ready"
-	return OutVolumeIndices.Num() == 0;
-	*/
 }
 
 const FUnreadyVolumeData_GT* FSpatialReadinessSimCallback::GetVolumeData_GT(int32 VolumeIndex) const
@@ -286,6 +233,14 @@ const FUnreadyVolumeData_GT* FSpatialReadinessSimCallback::GetVolumeData_GT(int3
 	return UnreadyVolumeData_GT.Find(VolumeIndex);
 }
 
+void FSpatialReadinessSimCallback::ForEachVolumeData_GT(const TFunction<void(const FUnreadyVolumeData_GT&)>& Func)
+{
+	for (int32 Index = 0; Index < UnreadyVolumeData_GT.Num(); ++Index)
+	{
+		Func(UnreadyVolumeData_GT.At(Index));
+	}
+}
+
 bool FSpatialReadinessSimCallback::QueryReadiness_PT(const FAABB3& Bounds, TArray<const FSingleParticlePhysicsProxy*>& OutVolumeProxies)
 {
 	// Get the evolution, and get the acceleration structure from it
diff --git a/Engine/Plugins/Experimental/SpatialReadiness/Source/SpatialReadiness/Private/SpatialReadiness/SpatialReadinessSubsystem.cpp b/Engine/Plugins/Experimental/SpatialReadiness/Source/SpatialReadiness/Private/SpatialReadiness/SpatialReadinessSubsystem.cpp
--- a/Engine/Plugins/Experimental/SpatialReadiness/Source/SpatialReadiness/Private/SpatialReadiness/SpatialReadinessSubsystem.cpp
+++ b/Engine/Plugins/Experimental/SpatialReadiness/Source/SpatialReadiness/Private/SpatialReadiness/SpatialReadinessSubsystem.cpp
@@ -4,6 +4,12 @@
 #include "Engine/World.h"
 #include "Physics/Experimental/PhysScene_Chaos.h"
 #include "PBDRigidsSolver.h"
+#include "PhysicsProxy/SingleParticlePhysicsProxy.h"
+#if ENABLE_DRAW_DEBUG
+#include "DisplayDebugHelpers.h"
+#include "DrawDebugHelpers.h"
+#include "GameFramework/HUD.h"
+#endif
 
 using namespace Chaos;
 
@@ -32,10 +38,18 @@ void USpatialReadiness::OnWorldBeginPlay(UWorld& InWorld)
 void USpatialReadiness::Initialize(FSubsystemCollectionBase& Collection)
 {
 	Super::Initialize(Collection);
+
+#if ENABLE_DRAW_DEBUG
+	ShowDebugInfoDelegateHandle = AHUD::OnShowDebugInfo.AddUObject(this, &USpatialReadiness::OnShowDebugInfo);
+#endif
 }
 
 void USpatialReadiness::Deinitialize()
 {
+#if ENABLE_DRAW_DEBUG
+	ShowDebugInfoDelegateHandle.Reset();
+#endif
+
 	DestroySimCallback();
 	Super::Deinitialize();
 }
@@ -176,3 +190,43 @@ Chaos::FPhysicsSolver* USpatialReadiness::GetSolver()
 
 	return nullptr;
 }
+
+#if ENABLE_DRAW_DEBUG
+
+void USpatialReadiness::OnShowDebugInfo(class AHUD* HUD, class UCanvas* Canvas, const FDebugDisplayInfo& DebugDisplay, float& YL, float& YPos)
+{
+	static const FName DisplayDebugName(TEXT("SpatialReadiness"));
+
+	if (DebugDisplay.IsDisplayOn(DisplayDebugName) == false)
+	{
+		return;
+	}
+
+	if (SimCallback == nullptr)
+	{
+		return;
+	}
+
+	UWorld* World = GetWorld();
+	if (World == nullptr)
+	{
+		return;
+	}
+
+	SimCallback->ForEachVolumeData_GT([World](const FUnreadyVolumeData_GT& Data)
+	{
+		FSingleParticlePhysicsProxy* Proxy = Data.Proxy;
+		if (Proxy == nullptr)
+		{
+			return;
+		}
+
+		//Chaos::FRigidBodyHandle_External& Handle = Proxy->GetGameThreadAPI();
+		Chaos::FGeometryParticleHandle* Handle = Proxy->GetHandle_LowLevel();
+		const FAABB3 Bounds = Handle->WorldSpaceInflatedBounds();
+
+		DrawDebugBox(World, Bounds.Center(), Bounds.Extents(), FQuat::Identity, FColor::Red, false, -1.f, -1, 0.f);
+	});
+}
+
+#endif
diff --git a/Engine/Plugins/Experimental/SpatialReadiness/Source/SpatialReadiness/Public/SpatialReadinessSimCallback.h b/Engine/Plugins/Experimental/SpatialReadiness/Source/SpatialReadiness/Public/SpatialReadinessSimCallback.h
--- a/Engine/Plugins/Experimental/SpatialReadiness/Source/SpatialReadiness/Public/SpatialReadinessSimCallback.h
+++ b/Engine/Plugins/Experimental/SpatialReadiness/Source/SpatialReadiness/Public/SpatialReadinessSimCallback.h
@@ -74,6 +74,9 @@ public:
 	// Given a volume index, get it's description
 	const FUnreadyVolumeData_GT* GetVolumeData_GT(int32 VolumeIndex) const;
 
+	// Iterate over each unready volume and get 
+	void ForEachVolumeData_GT(const TFunction<void(const FUnreadyVolumeData_GT&)>& Func);
+
 	/* begin: TSimCallbackObject */
 protected:
 	virtual void OnPreSimulate_Internal() override;
diff --git a/Engine/Plugins/Experimental/SpatialReadiness/Source/SpatialReadiness/Public/SpatialReadinessSubsystem.h b/Engine/Plugins/Experimental/SpatialReadiness/Source/SpatialReadiness/Public/SpatialReadinessSubsystem.h
--- a/Engine/Plugins/Experimental/SpatialReadiness/Source/SpatialReadiness/Public/SpatialReadinessSubsystem.h
+++ b/Engine/Plugins/Experimental/SpatialReadiness/Source/SpatialReadiness/Public/SpatialReadinessSubsystem.h
@@ -5,9 +5,11 @@
 #include "Subsystems/WorldSubsystem.h"
 #include "SpatialReadinessAPI.h"
 #include "SpatialReadinessSimCallback.h"
+#include "DrawDebugHelpers.h"
+#include "DisplayDebugHelpers.h"
 #include "SpatialReadinessSubsystem.generated.h"
 
-UCLASS()
+UCLASS(MinimalAPI)
 class USpatialReadiness : public UWorldSubsystem
 {
 	GENERATED_BODY()
@@ -18,14 +20,25 @@ public:
 	USpatialReadiness(FVTableHelper& Helper);
 
 	// Add a volume which can be marked read/unready
-	FSpatialReadinessVolume AddReadinessVolume(const FBox& Bounds, const FString& Description);
+	SPATIALREADINESS_API FSpatialReadinessVolume AddReadinessVolume(const FBox& Bounds, const FString& Description);
 
 	// Check to see if a volume is ready
 	//
 	// If bAllDescriptions = false, OutDescriptions will contain only the description of the
 	// first unready volume that we hit. If true, it will contain descriptions from all
 	// overlapping volumes.
-	bool QueryReadiness(const FBox& Bounds, TArray<FString>& OutDescriptions, bool bAllDescriptions = false) const;
+	SPATIALREADINESS_API bool QueryReadiness(const FBox& Bounds, TArray<FString>& OutDescriptions, bool bAllDescriptions = false) const;
+
+	//
+	// Debug functions
+	//
+
+#if ENABLE_DRAW_DEBUG
+
+	FDelegateHandle ShowDebugInfoDelegateHandle;
+	SPATIALREADINESS_API void OnShowDebugInfo(class AHUD* HUD, class UCanvas* Canvas, const FDebugDisplayInfo& DebugDisplay, float& YL, float& YPos);
+
+#endif
 
 	/* begin: UWorldSubsystem */
 protected:
