Fix query collision parameters for unready volumes so that they do not block scene queries.
Update scene query hit result conversions so that it's valid to have payloads with no userdata.
Update debug draw for spatial readiness.

diff --git a/Engine/Plugins/Experimental/SpatialReadiness/Config/DefaultInput.ini b/Engine/Plugins/Experimental/SpatialReadiness/Config/DefaultInput.ini
deleted file mode 100644
diff --git a/Engine/Plugins/Experimental/SpatialReadiness/Source/SpatialReadiness/Private/SpatialReadiness/SpatialReadinessDebug.cpp b/Engine/Plugins/Experimental/SpatialReadiness/Source/SpatialReadiness/Private/SpatialReadiness/SpatialReadinessDebug.cpp
new file mode 100644
--- /dev/null
+++ b/Engine/Plugins/Experimental/SpatialReadiness/Source/SpatialReadiness/Private/SpatialReadiness/SpatialReadinessDebug.cpp
@@ -0,0 +1,9 @@
+// Copyright Epic Games, Inc. All Rights Reserved.
+
+#include "SpatialReadinessDebug.h"
+
+TAutoConsoleVariable<bool> CVarSpatialReadinessDebugDraw(
+	TEXT("p.SpatialReadiness.DebugDraw"),
+	0,
+	TEXT("Debug drawing for spatial readiness subsystem"),
+	ECVF_Default);
diff --git a/Engine/Plugins/Experimental/SpatialReadiness/Source/SpatialReadiness/Private/SpatialReadiness/SpatialReadinessSimCallback.cpp b/Engine/Plugins/Experimental/SpatialReadiness/Source/SpatialReadiness/Private/SpatialReadiness/SpatialReadinessSimCallback.cpp
--- a/Engine/Plugins/Experimental/SpatialReadiness/Source/SpatialReadiness/Private/SpatialReadiness/SpatialReadinessSimCallback.cpp
+++ b/Engine/Plugins/Experimental/SpatialReadiness/Source/SpatialReadiness/Private/SpatialReadiness/SpatialReadinessSimCallback.cpp
@@ -18,7 +18,10 @@
 #include "SpatialReadinessStats.h"
 #include "Engine/OverlapResult.h"
 #include "PhysicsEngine/PhysicsObjectExternalInterface.h"
-
+#include "DrawDebugHelpers.h"
+#include "Chaos/ChaosDebugDraw.h"
+#include "Chaos/DebugDrawQueue.h"
+#include "SpatialReadinessDebug.h"
 
 using namespace Chaos;
 
@@ -72,12 +75,14 @@ int32 FSpatialReadinessSimCallback::AddUnreadyVolume_GT(const FBox& Bounds, cons
 
 	// Create collision response container
 	FCollisionResponseContainer CollisionResponse;
+	// We want this to block everything that simulates, even if it isn't in the ECC_PhysicsBody channel.
+	// Since bodies can be in any channel they want, we have to block all channels.
 	CollisionResponse.SetAllChannels(ECollisionResponse::ECR_Block);
 
 	// Create collision filter data for the particle
 	FCollisionFilterData QueryFilterData, SimFilterData;
 	CreateShapeFilterData(
-		/* MyChannel */            static_cast<uint8>(ECollisionChannel::ECC_WorldStatic),
+		/* MyChannel */            static_cast<uint8>(ECollisionChannel::ECC_WorldDynamic),
 		/* MaskFilter */           FMaskFilter(0),
 		/* SourceObjectID */       0,
 		/* ResponseToChannels */   CollisionResponse,
@@ -88,7 +93,15 @@ int32 FSpatialReadinessSimCallback::AddUnreadyVolume_GT(const FBox& Bounds, cons
 		/* bEnableCCD */           true,
 		/* bEnableContactNotify */ false,
 		/* bStaticShape */         true);
-	//SimFilterData.Word3 |= EPDF_SimpleCollision | EPDF_ComplexCollision;
+
+	// Notes of query filter data:
+	// * We want to block all channels in sim but overlap all channels query.
+	// * By default CreateShapeFilterData uses CollisionResponse for both
+	// * Word1 is blocking data, and Word2 is overlap data.
+	// * So we OR query Word1 into Word2 so that we block overlaps and clear
+	//   Word1 so that we don't block any queries.
+	QueryFilterData.Word2 |= QueryFilterData.Word1;
+	QueryFilterData.Word1 = 0;
 	QueryFilterData.Word3 |= EPDF_SimpleCollision | EPDF_ComplexCollision;
 
 	// Make the geometry
@@ -162,7 +175,7 @@ void FSpatialReadinessSimCallback::RemoveUnreadyVolume_GT(int32 UnreadyVolumeInd
 bool FSpatialReadinessSimCallback::QueryReadiness_GT(const FBox& Bounds, TArray<int32>& OutVolumeIndices, bool bAllUnreadyVolumes) const
 {
 	// Constants that we'll use to set up query parameters
-	constexpr static ECollisionChannel Channel = ECollisionChannel::ECC_WorldStatic;
+	constexpr static ECollisionChannel Channel = ECollisionChannel::ECC_PhysicsBody;//Static;
 	constexpr static uint8 ChannelBit = static_cast<uint8>(Channel);
 	constexpr static bool bComplex = false;
 	const bool bMulti = bAllUnreadyVolumes;
@@ -452,6 +465,14 @@ void FSpatialReadinessSimCallback::OnMidPhaseModification_Internal(FMidPhaseModi
 			continue;
 		}
 
+#if ENABLE_DRAW_DEBUG
+		if (CVarSpatialReadinessDebugDraw.GetValueOnAnyThread())
+		{
+			const FAABB3 Bounds = UnreadyVolume->WorldSpaceInflatedBounds();
+			Chaos::FDebugDrawQueue::GetInstance().DrawDebugBox(Bounds.Center(), Bounds.Extents() * .5f, FQuat::Identity, FColor::Red, false, -1.f, -1, 0.f);
+		}
+#endif
+
 		// Go through every mid-phase which involves this volume
 		for (FMidPhaseModifier& MidPhase : Accessor.GetMidPhases(UnreadyVolume))
 		{
@@ -468,6 +489,14 @@ void FSpatialReadinessSimCallback::OnMidPhaseModification_Internal(FMidPhaseModi
 				continue;
 			}
 
+#if ENABLE_DRAW_DEBUG
+			if (CVarSpatialReadinessDebugDraw.GetValueOnAnyThread())
+			{
+				const FAABB3 Bounds = RigidParticle->WorldSpaceInflatedBounds();
+				Chaos::FDebugDrawQueue::GetInstance().DrawDebugBox(Bounds.Center(), Bounds.Extents() * .5f, FQuat::Identity, FColor::Yellow, false, -1.f, -1, 0.f);
+			}
+#endif
+
 			// Add the particle to the list of particles to freeze
 			UnreadyRigidParticles_PT.Add(RigidParticle);
 
diff --git a/Engine/Plugins/Experimental/SpatialReadiness/Source/SpatialReadiness/Private/SpatialReadiness/SpatialReadinessSubsystem.cpp b/Engine/Plugins/Experimental/SpatialReadiness/Source/SpatialReadiness/Private/SpatialReadiness/SpatialReadinessSubsystem.cpp
--- a/Engine/Plugins/Experimental/SpatialReadiness/Source/SpatialReadiness/Private/SpatialReadiness/SpatialReadinessSubsystem.cpp
+++ b/Engine/Plugins/Experimental/SpatialReadiness/Source/SpatialReadiness/Private/SpatialReadiness/SpatialReadinessSubsystem.cpp
@@ -5,11 +5,6 @@
 #include "Physics/Experimental/PhysScene_Chaos.h"
 #include "PBDRigidsSolver.h"
 #include "PhysicsProxy/SingleParticlePhysicsProxy.h"
-#if ENABLE_DRAW_DEBUG
-#include "DisplayDebugHelpers.h"
-#include "DrawDebugHelpers.h"
-#include "GameFramework/HUD.h"
-#endif
 
 using namespace Chaos;
 
@@ -38,18 +33,10 @@ void USpatialReadiness::OnWorldBeginPlay(UWorld& InWorld)
 void USpatialReadiness::Initialize(FSubsystemCollectionBase& Collection)
 {
 	Super::Initialize(Collection);
-
-#if ENABLE_DRAW_DEBUG
-	ShowDebugInfoDelegateHandle = AHUD::OnShowDebugInfo.AddUObject(this, &USpatialReadiness::OnShowDebugInfo);
-#endif
 }
 
 void USpatialReadiness::Deinitialize()
 {
-#if ENABLE_DRAW_DEBUG
-	ShowDebugInfoDelegateHandle.Reset();
-#endif
-
 	DestroySimCallback();
 	Super::Deinitialize();
 }
@@ -191,42 +178,3 @@ Chaos::FPhysicsSolver* USpatialReadiness::GetSolver()
 	return nullptr;
 }
 
-#if ENABLE_DRAW_DEBUG
-
-void USpatialReadiness::OnShowDebugInfo(class AHUD* HUD, class UCanvas* Canvas, const FDebugDisplayInfo& DebugDisplay, float& YL, float& YPos)
-{
-	static const FName DisplayDebugName(TEXT("SpatialReadiness"));
-
-	if (DebugDisplay.IsDisplayOn(DisplayDebugName) == false)
-	{
-		return;
-	}
-
-	if (SimCallback == nullptr)
-	{
-		return;
-	}
-
-	UWorld* World = GetWorld();
-	if (World == nullptr)
-	{
-		return;
-	}
-
-	SimCallback->ForEachVolumeData_GT([World](const FUnreadyVolumeData_GT& Data)
-	{
-		FSingleParticlePhysicsProxy* Proxy = Data.Proxy;
-		if (Proxy == nullptr)
-		{
-			return;
-		}
-
-		//Chaos::FRigidBodyHandle_External& Handle = Proxy->GetGameThreadAPI();
-		Chaos::FGeometryParticleHandle* Handle = Proxy->GetHandle_LowLevel();
-		const FAABB3 Bounds = Handle->WorldSpaceInflatedBounds();
-
-		DrawDebugBox(World, Bounds.Center(), Bounds.Extents(), FQuat::Identity, FColor::Red, false, -1.f, -1, 0.f);
-	});
-}
-
-#endif
diff --git a/Engine/Plugins/Experimental/SpatialReadiness/Source/SpatialReadiness/Public/SpatialReadinessDebug.h b/Engine/Plugins/Experimental/SpatialReadiness/Source/SpatialReadiness/Public/SpatialReadinessDebug.h
new file mode 100644
--- /dev/null
+++ b/Engine/Plugins/Experimental/SpatialReadiness/Source/SpatialReadiness/Public/SpatialReadinessDebug.h
@@ -0,0 +1,6 @@
+// Copyright Epic Games, Inc. All Rights Reserved.
+#pragma once
+#include "CoreMinimal.h"
+#include "HAL/IConsoleManager.h"
+
+extern TAutoConsoleVariable<bool> CVarSpatialReadinessDebugDraw;
diff --git a/Engine/Plugins/Experimental/SpatialReadiness/Source/SpatialReadiness/Public/SpatialReadinessSubsystem.h b/Engine/Plugins/Experimental/SpatialReadiness/Source/SpatialReadiness/Public/SpatialReadinessSubsystem.h
--- a/Engine/Plugins/Experimental/SpatialReadiness/Source/SpatialReadiness/Public/SpatialReadinessSubsystem.h
+++ b/Engine/Plugins/Experimental/SpatialReadiness/Source/SpatialReadiness/Public/SpatialReadinessSubsystem.h
@@ -5,8 +5,6 @@
 #include "Subsystems/WorldSubsystem.h"
 #include "SpatialReadinessAPI.h"
 #include "SpatialReadinessSimCallback.h"
-#include "DrawDebugHelpers.h"
-#include "DisplayDebugHelpers.h"
 #include "SpatialReadinessSubsystem.generated.h"
 
 UCLASS(MinimalAPI)
@@ -29,17 +27,6 @@ public:
 	// overlapping volumes.
 	SPATIALREADINESS_API bool QueryReadiness(const FBox& Bounds, TArray<FString>& OutDescriptions, bool bAllDescriptions = false) const;
 
-	//
-	// Debug functions
-	//
-
-#if ENABLE_DRAW_DEBUG
-
-	FDelegateHandle ShowDebugInfoDelegateHandle;
-	SPATIALREADINESS_API void OnShowDebugInfo(class AHUD* HUD, class UCanvas* Canvas, const FDebugDisplayInfo& DebugDisplay, float& YL, float& YPos);
-
-#endif
-
 	/* begin: UWorldSubsystem */
 protected:
 	virtual bool ShouldCreateSubsystem(UObject* Outer) const override;
diff --git a/Engine/Source/Runtime/Engine/Private/Collision/CollisionConversions.cpp b/Engine/Source/Runtime/Engine/Private/Collision/CollisionConversions.cpp
--- a/Engine/Source/Runtime/Engine/Private/Collision/CollisionConversions.cpp
+++ b/Engine/Source/Runtime/Engine/Private/Collision/CollisionConversions.cpp
@@ -140,7 +140,12 @@ static void SetHitResultFromShapeAndFaceIndex(const FPhysicsShape& Shape,  const
 	IPhysicsBodyInstanceOwner* BodyInstanceOwner = nullptr;
 	TWeakObjectPtr<UPrimitiveComponent> OwningComponent;
 	OutResult.PhysicsObject = nullptr;
-	if(const FBodyInstance* BodyInst = GetUserData(Actor))
+	
+	if (!HasValidUserData(Actor))
+	{
+		return;
+	}
+	else if (const FBodyInstance* BodyInst = GetUserData(Actor))
 	{
 		BodyInst = FPhysicsInterface::ShapeToOriginalBodyInstance(BodyInst, &Shape);
 
@@ -666,10 +671,15 @@ void ConvertQueryOverlap(const FPhysicsShape& Shape, const FPhysicsActor& Actor,
 
 	// Grab actor/component
 	OutOverlap.PhysicsObject = nullptr;
+	
+	if (!HasValidUserData(Actor))
+	{
+		return;
+	}
 	// Try body instance
-	if (const FBodyInstance* BodyInst = GetUserData(Actor))
+	else if (const FBodyInstance* BodyInst = GetUserData(Actor))
 	{
-        BodyInst = FPhysicsInterface::ShapeToOriginalBodyInstance(BodyInst, &Shape);
+		BodyInst = FPhysicsInterface::ShapeToOriginalBodyInstance(BodyInst, &Shape);
 		if (const UPrimitiveComponent* OwnerComponent = BodyInst->OwnerComponent.Get())
 		{
 			// Create a handle that needs to be resolved later when accessed in game thread,
diff --git a/Engine/Source/Runtime/Engine/Private/PhysicsEngine/Experimental/ChaosInterfaceWrapper.cpp b/Engine/Source/Runtime/Engine/Private/PhysicsEngine/Experimental/ChaosInterfaceWrapper.cpp
--- a/Engine/Source/Runtime/Engine/Private/PhysicsEngine/Experimental/ChaosInterfaceWrapper.cpp
+++ b/Engine/Source/Runtime/Engine/Private/PhysicsEngine/Experimental/ChaosInterfaceWrapper.cpp
@@ -53,6 +53,16 @@ namespace ChaosInterface
 		return nullptr;
 	}
 
+	bool HasValidUserData(const Chaos::FGeometryParticle& Actor)
+	{
+		if (FChaosUserData* UserData = static_cast<FChaosUserData*>(Actor.UserData()))
+		{
+			return !UserData->Is<EChaosUserDataType::Invalid>();
+		}
+
+		return false;
+	}
+
 	FScopedSceneReadLock::FScopedSceneReadLock(FPhysScene_Chaos& SceneIn)
 		: Solver(SceneIn.GetSolver())
 	{
diff --git a/Engine/Source/Runtime/Engine/Public/Physics/Experimental/ChaosInterfaceWrapper.h b/Engine/Source/Runtime/Engine/Public/Physics/Experimental/ChaosInterfaceWrapper.h
--- a/Engine/Source/Runtime/Engine/Public/Physics/Experimental/ChaosInterfaceWrapper.h
+++ b/Engine/Source/Runtime/Engine/Public/Physics/Experimental/ChaosInterfaceWrapper.h
@@ -42,6 +42,7 @@ inline FQueryFilterData MakeQueryFilterData(const FCollisionFilterData& FilterDa
 FBodyInstance* GetUserData(const Chaos::FGeometryParticle& Actor);
 UPhysicalMaterial* GetUserData(const Chaos::FChaosPhysicsMaterial& Material);
 UPrimitiveComponent* GetPrimitiveComponentFromUserData(const Chaos::FGeometryParticle& Actor);
+bool HasValidUserData(const Chaos::FGeometryParticle& Actor);
 
 }
 
diff --git a/Engine/Source/Runtime/PhysicsCore/Public/PhysicsInterfaceTypesCore.h b/Engine/Source/Runtime/PhysicsCore/Public/PhysicsInterfaceTypesCore.h
--- a/Engine/Source/Runtime/PhysicsCore/Public/PhysicsInterfaceTypesCore.h
+++ b/Engine/Source/Runtime/PhysicsCore/Public/PhysicsInterfaceTypesCore.h
@@ -218,6 +218,9 @@ public:
 
 	//helper function to determine if userData is garbage (maybe dangling pointer)
 	static bool IsGarbage(void* UserData){ return ((FChaosUserData*)UserData)->Type < EChaosUserDataType::Invalid || ((FChaosUserData*)UserData)->Type > EChaosUserDataType::CustomPayload; }
+
+	template <EChaosUserDataType::Type UserDataType>
+	bool Is() const { return Type == UserDataType; }
 };
 
 using FUserData = FChaosUserData;
