Fix move semantics on FSpatialReadinessVolume

diff --git a/Engine/Plugins/Experimental/SpatialReadiness/Source/SpatialReadiness/Private/SpatialReadiness/SpatialReadinessAPI.cpp b/Engine/Plugins/Experimental/SpatialReadiness/Source/SpatialReadiness/Private/SpatialReadiness/SpatialReadinessAPI.cpp
--- a/Engine/Plugins/Experimental/SpatialReadiness/Source/SpatialReadiness/Private/SpatialReadiness/SpatialReadinessAPI.cpp
+++ b/Engine/Plugins/Experimental/SpatialReadiness/Source/SpatialReadiness/Private/SpatialReadiness/SpatialReadinessAPI.cpp
@@ -24,13 +24,13 @@ FSpatialReadinessAPI::~FSpatialReadinessAPI()
 FSpatialReadinessVolume FSpatialReadinessAPI::CreateVolume(const FBox& Bounds, const FString& Description)
 {
 	// Create a handle for a new unready volume
-	FSpatialReadinessVolume Handle(Delegates.ToWeakPtr(), Bounds, Description);
+	FSpatialReadinessVolume ReadinessVolume(Delegates.ToWeakPtr(), Bounds, Description);
 
 	// Default to an unready state. This should call back into AddUnreadyVolume_Internal
 	// to create the underlying unready volume.
-	Handle.MarkUnready();
+	ReadinessVolume.MarkUnready();
 
 	// Return the newly created handle. This handle will only be valid so long as this
 	// ISpatialReadiness object is in scope.
-	return Handle;
+	return ReadinessVolume;
 }
diff --git a/Engine/Plugins/Experimental/SpatialReadiness/Source/SpatialReadiness/Private/SpatialReadiness/SpatialReadinessVolume.cpp b/Engine/Plugins/Experimental/SpatialReadiness/Source/SpatialReadiness/Private/SpatialReadiness/SpatialReadinessVolume.cpp
--- a/Engine/Plugins/Experimental/SpatialReadiness/Source/SpatialReadiness/Private/SpatialReadiness/SpatialReadinessVolume.cpp
+++ b/Engine/Plugins/Experimental/SpatialReadiness/Source/SpatialReadiness/Private/SpatialReadiness/SpatialReadinessVolume.cpp
@@ -31,6 +31,21 @@ FSpatialReadinessVolume::FSpatialReadinessVolume(FSpatialReadinessVolume&& Other
 #endif
 }
 
+FSpatialReadinessVolume& FSpatialReadinessVolume::operator=(FSpatialReadinessVolume&& Other)
+{
+	Delegates = Other.Delegates;
+	Index = Other.Index;
+	Bounds = Other.Bounds;
+	Other.Delegates.Reset();
+	Other.Index = INDEX_NONE;
+	Other.Bounds = FBox();
+#if WITH_SPATIAL_READINESS_DESCRIPTIONS
+	Description = Other.Description;
+	Other.Description = TEXT("");
+#endif
+	return *this;
+}
+
 FSpatialReadinessVolume::~FSpatialReadinessVolume()
 {
 	// When a handle is destroyed, we mark the corresponding volume as
diff --git a/Engine/Plugins/Experimental/SpatialReadiness/Source/SpatialReadiness/Public/SpatialReadinessVolume.h b/Engine/Plugins/Experimental/SpatialReadiness/Source/SpatialReadiness/Public/SpatialReadinessVolume.h
--- a/Engine/Plugins/Experimental/SpatialReadiness/Source/SpatialReadiness/Public/SpatialReadinessVolume.h
+++ b/Engine/Plugins/Experimental/SpatialReadiness/Source/SpatialReadiness/Public/SpatialReadinessVolume.h
@@ -46,10 +46,12 @@ public:
 	SPATIALREADINESS_API FSpatialReadinessVolume() = default;
 
 	// Readiness handles are unique and cannot be copied
-	FSpatialReadinessVolume(const FSpatialReadinessVolume&) = delete;
+	SPATIALREADINESS_API FSpatialReadinessVolume(const FSpatialReadinessVolume&) = delete;
+	SPATIALREADINESS_API FSpatialReadinessVolume& operator=(const FSpatialReadinessVolume&) = delete;
 
-	// Move constructor
+	// Move constructor & assignment
 	SPATIALREADINESS_API FSpatialReadinessVolume(FSpatialReadinessVolume&& Other);
+	SPATIALREADINESS_API FSpatialReadinessVolume& operator=(FSpatialReadinessVolume&& Other);
 
 private:
 
