remove exports on deleted functions for horde

diff --git a/Engine/Plugins/Experimental/SpatialReadiness/Source/SpatialReadiness/Public/SpatialReadinessVolume.h b/Engine/Plugins/Experimental/SpatialReadiness/Source/SpatialReadiness/Public/SpatialReadinessVolume.h
--- a/Engine/Plugins/Experimental/SpatialReadiness/Source/SpatialReadiness/Public/SpatialReadinessVolume.h
+++ b/Engine/Plugins/Experimental/SpatialReadiness/Source/SpatialReadiness/Public/SpatialReadinessVolume.h
@@ -46,8 +46,8 @@ public:
 	SPATIALREADINESS_API FSpatialReadinessVolume() = default;
 
 	// Readiness handles are unique and cannot be copied
-	SPATIALREADINESS_API FSpatialReadinessVolume(const FSpatialReadinessVolume&) = delete;
-	SPATIALREADINESS_API FSpatialReadinessVolume& operator=(const FSpatialReadinessVolume&) = delete;
+	FSpatialReadinessVolume(const FSpatialReadinessVolume&) = delete;
+	FSpatialReadinessVolume& operator=(const FSpatialReadinessVolume&) = delete;
 
 	// Move constructor & assignment
 	SPATIALREADINESS_API FSpatialReadinessVolume(FSpatialReadinessVolume&& Other);
