Fix issue with SpatialReadinessVolume MarkReady. Also fix const correctness

diff --git a/Engine/Plugins/Experimental/SpatialReadiness/Source/SpatialReadiness/Private/SpatialReadiness/SpatialReadinessVolume.cpp b/Engine/Plugins/Experimental/SpatialReadiness/Source/SpatialReadiness/Private/SpatialReadiness/SpatialReadinessVolume.cpp
--- a/Engine/Plugins/Experimental/SpatialReadiness/Source/SpatialReadiness/Private/SpatialReadiness/SpatialReadinessVolume.cpp
+++ b/Engine/Plugins/Experimental/SpatialReadiness/Source/SpatialReadiness/Private/SpatialReadiness/SpatialReadinessVolume.cpp
@@ -88,6 +88,11 @@ bool FSpatialReadinessVolume::IsReady() const
 
 void FSpatialReadinessVolume::MarkReady()
 {
+	if (!IsValid())
+	{
+		return;
+	}
+
 	if (IsReady())
 	{
 		return;
@@ -131,7 +136,7 @@ void FSpatialReadinessVolume::MarkUnready()
 	}
 }
 
-void FSpatialReadinessVolume::MarkUnready(FBox& NewBounds)
+void FSpatialReadinessVolume::MarkUnready(const FBox& NewBounds)
 {
 	// Store new bounds
 	Bounds = NewBounds;
diff --git a/Engine/Plugins/Experimental/SpatialReadiness/Source/SpatialReadiness/Public/SpatialReadinessVolume.h b/Engine/Plugins/Experimental/SpatialReadiness/Source/SpatialReadiness/Public/SpatialReadinessVolume.h
--- a/Engine/Plugins/Experimental/SpatialReadiness/Source/SpatialReadiness/Public/SpatialReadinessVolume.h
+++ b/Engine/Plugins/Experimental/SpatialReadiness/Source/SpatialReadiness/Public/SpatialReadinessVolume.h
@@ -33,7 +33,7 @@ public:
 	SPATIALREADINESS_API void MarkUnready();
 
 	// Optional override to MakeUnready() which adjusts the bounds
-	SPATIALREADINESS_API void MarkUnready(FBox& NewBounds);
+	SPATIALREADINESS_API void MarkUnready(const FBox& NewBounds);
 
 	// Get the description of this handle - will be empty string if
 	// the WITH_SPATIAL_READINESS_DESCRIPTIONS macro is 0
