Add logging and CVD debug drawing for spatial readiness

diff --git a/Engine/Plugins/Experimental/SpatialReadiness/Source/SpatialReadiness/Private/SpatialReadiness/SpatialReadinessDebug.cpp b/Engine/Plugins/Experimental/SpatialReadiness/Source/SpatialReadiness/Private/SpatialReadiness/SpatialReadinessDebug.cpp
--- a/Engine/Plugins/Experimental/SpatialReadiness/Source/SpatialReadiness/Private/SpatialReadiness/SpatialReadinessDebug.cpp
+++ b/Engine/Plugins/Experimental/SpatialReadiness/Source/SpatialReadiness/Private/SpatialReadiness/SpatialReadinessDebug.cpp
@@ -1,9 +1,73 @@
 // Copyright Epic Games, Inc. All Rights Reserved.
 
 #include "SpatialReadinessDebug.h"
+#include "SpatialReadinessLog.h"
+#include "SpatialReadinessSubsystem.h"
+#include "SpatialReadinessSimCallback.h"
+#include "PhysicsProxy/SingleParticlePhysicsProxy.h"
+#include "Chaos/ParticleHandle.h"
+#include "Engine/World.h"
 
 TAutoConsoleVariable<bool> CVarSpatialReadinessDebugDraw(
 	TEXT("p.SpatialReadiness.DebugDraw"),
 	0,
 	TEXT("Debug drawing for spatial readiness subsystem"),
 	ECVF_Default);
+
+
+static FAutoConsoleCommandWithWorld CmdSpatialReadinessPrint(
+	TEXT("p.SpatialReadiness.Print"),
+	TEXT("Command to log all current unready volumes and their statuses"),
+	FConsoleCommandWithWorldDelegate::CreateStatic([](UWorld* World)
+	{
+#if !UE_BUILD_SHIPPING
+		if (World == nullptr)
+		{
+			UE_LOG(LogSpatialReadiness, Log, TEXT("Can't print spatial readiness: no world"));
+			return;
+		}
+
+		USpatialReadiness* SpatialReadiness = World->GetSubsystem<USpatialReadiness>();
+		if (SpatialReadiness == nullptr)
+		{
+			UE_LOG(LogSpatialReadiness, Log, TEXT("Can't print spatial readiness: no spatial readiness subsystem"));
+			return;
+		}
+
+		FSpatialReadinessSimCallback* SimCallback = SpatialReadiness->GetSimCallback();
+		if (SimCallback == nullptr)
+		{
+			UE_LOG(LogSpatialReadiness, Log, TEXT("Can't print spatial readiness: no spatial readiness sim callback"));
+			return;
+		}
+
+		int32 VolumeNum = 0;
+		TStringBuilder<(1<<14)> Builder;
+		Builder << TEXT("Unready Volumes (" << SimCallback->NumUnreadyVolumes_GT() << "):\n");
+		SimCallback->ForEachVolumeData_GT([&Builder, &VolumeNum](const FUnreadyVolumeData_GT& Data)
+		{
+			FSingleParticlePhysicsProxy* Proxy = Data.Proxy;
+			if (Proxy == nullptr)
+			{
+				return;
+			}
+
+			if (Proxy->GetMarkedDeleted())
+			{
+				return;
+			}
+
+			FGeometryParticleHandle* UnreadyVolume = Proxy->GetHandle_LowLevel();
+			if (UnreadyVolume == nullptr)
+			{
+				return;
+			}
+
+			const FAABB3 Bounds = UnreadyVolume->WorldSpaceInflatedBounds();
+
+			Builder << " " << (VolumeNum++) << ". " << Data.Description << " : " << *Bounds.ToString() << "\n";
+		});
+
+		UE_LOG(LogSpatialReadiness, Log, TEXT("%s"), *Builder);
+#endif
+	}));
diff --git a/Engine/Plugins/Experimental/SpatialReadiness/Source/SpatialReadiness/Private/SpatialReadiness/SpatialReadinessLog.cpp b/Engine/Plugins/Experimental/SpatialReadiness/Source/SpatialReadiness/Private/SpatialReadiness/SpatialReadinessLog.cpp
new file mode 100644
--- /dev/null
+++ b/Engine/Plugins/Experimental/SpatialReadiness/Source/SpatialReadiness/Private/SpatialReadiness/SpatialReadinessLog.cpp
@@ -0,0 +1,3 @@
+#include "SpatialReadinessLog.h"
+
+DEFINE_LOG_CATEGORY(LogSpatialReadiness);
\ No newline at end of file
diff --git a/Engine/Plugins/Experimental/SpatialReadiness/Source/SpatialReadiness/Private/SpatialReadiness/SpatialReadinessSimCallback.cpp b/Engine/Plugins/Experimental/SpatialReadiness/Source/SpatialReadiness/Private/SpatialReadiness/SpatialReadinessSimCallback.cpp
--- a/Engine/Plugins/Experimental/SpatialReadiness/Source/SpatialReadiness/Private/SpatialReadiness/SpatialReadinessSimCallback.cpp
+++ b/Engine/Plugins/Experimental/SpatialReadiness/Source/SpatialReadiness/Private/SpatialReadiness/SpatialReadinessSimCallback.cpp
@@ -21,6 +21,7 @@
 #include "DrawDebugHelpers.h"
 #include "Chaos/ChaosDebugDraw.h"
 #include "Chaos/DebugDrawQueue.h"
+#include "SpatialReadinessLog.h"
 #include "SpatialReadinessDebug.h"
 
 using namespace Chaos;
@@ -112,7 +113,7 @@ int32 FSpatialReadinessSimCallback::AddUnreadyVolume_GT(const FBox& Bounds, cons
 	ParticleHandle.ShapesArray()[0]->SetIsProbe(true);
 	ParticleHandle.ShapesArray()[0]->SetQueryData(QueryFilterData);
 #if CHAOS_DEBUG_NAME
-	ParticleHandle.SetDebugName(MakeShared<FString>(TEXT("Unready Volume")));
+	ParticleHandle.SetDebugName(MakeShared<FString>(FString::Printf(TEXT("UnreadyVolume: %s"), *Description)));
 #endif
 
 	// Add the new particle to the scene
@@ -137,6 +138,9 @@ int32 FSpatialReadinessSimCallback::AddUnreadyVolume_GT(const FBox& Bounds, cons
 	SimInput.UnreadyVolumesToAdd.Add(ParticleProxy);
 	SimInput.UnreadyVolumesToRemove.Remove(ParticleProxy);
 
+	// Log the event
+	UE_LOG(LogSpatialReadiness, Log, TEXT("Unready volume created.with index [%d]: %s"), VolumeIndex, *Description);
+
 	return VolumeIndex;
 }
 
@@ -156,9 +160,16 @@ void FSpatialReadinessSimCallback::RemoveUnreadyVolume_GT(int32 UnreadyVolumeInd
 	{
 		return;
 	}
-	
-	// Free the index in our GT tracker
-	UnreadyVolumeData_GT.Remove(UnreadyVolumeIndex);
+
+	if (ensureMsgf(UnreadyVolumeData_GT.Find(UnreadyVolumeIndex), TEXT("No volume data associated with unready volume index")))
+	{
+		// Free the index in our GT tracker
+		UnreadyVolumeData_GT.Remove(UnreadyVolumeIndex);
+	}
+	else
+	{
+		UE_LOG(LogSpatialReadiness, Warning, TEXT("Attempted to remove unready volume with index %d, but it wasn't found!"), UnreadyVolumeIndex);
+	}
 
 	// Tell the PT to remove it's tracking of this proxy as well
 	FSpatialReadinessSimCallbackInput* SimInput = GetProducerInputData_External();
@@ -170,6 +181,9 @@ void FSpatialReadinessSimCallback::RemoveUnreadyVolume_GT(int32 UnreadyVolumeInd
 
 	// Delete the particle
 	FChaosEngineInterface::ReleaseActor(ParticleProxy, &PhysicsScene);
+
+	// Log the event
+	UE_LOG(LogSpatialReadiness, Log, TEXT("Unready volume removed with index [%d]"), UnreadyVolumeIndex);
 }
 
 bool FSpatialReadinessSimCallback::QueryReadiness_GT(const FBox& Bounds, TArray<int32>& OutVolumeIndices, bool bAllUnreadyVolumes) const
@@ -254,6 +268,11 @@ void FSpatialReadinessSimCallback::ForEachVolumeData_GT(const TFunction<void(con
 	}
 }
 
+int32 FSpatialReadinessSimCallback::NumUnreadyVolumes_GT() const
+{
+	return UnreadyVolumeData_GT.Num();
+}
+
 bool FSpatialReadinessSimCallback::QueryReadiness_PT(const FAABB3& Bounds, TArray<const FSingleParticlePhysicsProxy*>& OutVolumeProxies)
 {
 	// Get the evolution, and get the acceleration structure from it
@@ -473,6 +492,16 @@ void FSpatialReadinessSimCallback::OnMidPhaseModification_Internal(FMidPhaseModi
 		}
 #endif
 
+#if !UE_BUILD_SHIPPING
+		{
+			// Draw a box in CVD representing the unready volume
+			const FName Tag = *UnreadyVolume->GetDebugName();
+			const FLinearColor Color = FLinearColor::Red;
+			const FAABB3 Bounds = UnreadyVolume->WorldSpaceInflatedBounds();
+			CVD_TRACE_DEBUG_DRAW_BOX(FBox(Bounds.Min(), Bounds.Max()), Tag, Color.ToFColorSRGB(), FChaosVisualDebuggerTrace::GetSolverID(*GetSolver()));
+		}
+#endif
+
 		// Go through every mid-phase which involves this volume
 		for (FMidPhaseModifier& MidPhase : Accessor.GetMidPhases(UnreadyVolume))
 		{
diff --git a/Engine/Plugins/Experimental/SpatialReadiness/Source/SpatialReadiness/Private/SpatialReadiness/SpatialReadinessSubsystem.cpp b/Engine/Plugins/Experimental/SpatialReadiness/Source/SpatialReadiness/Private/SpatialReadiness/SpatialReadinessSubsystem.cpp
--- a/Engine/Plugins/Experimental/SpatialReadiness/Source/SpatialReadiness/Private/SpatialReadiness/SpatialReadinessSubsystem.cpp
+++ b/Engine/Plugins/Experimental/SpatialReadiness/Source/SpatialReadiness/Private/SpatialReadiness/SpatialReadinessSubsystem.cpp
@@ -5,6 +5,8 @@
 #include "Physics/Experimental/PhysScene_Chaos.h"
 #include "PBDRigidsSolver.h"
 #include "PhysicsProxy/SingleParticlePhysicsProxy.h"
+#include "ChaosVisualDebugger/ChaosVDTraceMacros.h"
+#include "ChaosVisualDebugger/ChaosVisualDebuggerTrace.h"
 
 using namespace Chaos;
 
diff --git a/Engine/Plugins/Experimental/SpatialReadiness/Source/SpatialReadiness/Public/SpatialReadinessLog.h b/Engine/Plugins/Experimental/SpatialReadiness/Source/SpatialReadiness/Public/SpatialReadinessLog.h
new file mode 100644
--- /dev/null
+++ b/Engine/Plugins/Experimental/SpatialReadiness/Source/SpatialReadiness/Public/SpatialReadinessLog.h
@@ -0,0 +1,5 @@
+// Copyright Epic Games, Inc. All Rights Reserved.
+#pragma once
+#include "Logging/LogMacros.h"
+
+DECLARE_LOG_CATEGORY_EXTERN(LogSpatialReadiness, Log, All);
diff --git a/Engine/Plugins/Experimental/SpatialReadiness/Source/SpatialReadiness/Public/SpatialReadinessSimCallback.h b/Engine/Plugins/Experimental/SpatialReadiness/Source/SpatialReadiness/Public/SpatialReadinessSimCallback.h
--- a/Engine/Plugins/Experimental/SpatialReadiness/Source/SpatialReadiness/Public/SpatialReadinessSimCallback.h
+++ b/Engine/Plugins/Experimental/SpatialReadiness/Source/SpatialReadiness/Public/SpatialReadinessSimCallback.h
@@ -77,6 +77,9 @@ public:
 	// Iterate over each unready volume and get 
 	void ForEachVolumeData_GT(const TFunction<void(const FUnreadyVolumeData_GT&)>& Func);
 
+	// Get the number of unready volumes
+	int32 NumUnreadyVolumes_GT() const;
+
 	/* begin: TSimCallbackObject */
 protected:
 	virtual void OnPreSimulate_Internal() override;
diff --git a/Engine/Plugins/Experimental/SpatialReadiness/Source/SpatialReadiness/Public/SpatialReadinessSubsystem.h b/Engine/Plugins/Experimental/SpatialReadiness/Source/SpatialReadiness/Public/SpatialReadinessSubsystem.h
--- a/Engine/Plugins/Experimental/SpatialReadiness/Source/SpatialReadiness/Public/SpatialReadinessSubsystem.h
+++ b/Engine/Plugins/Experimental/SpatialReadiness/Source/SpatialReadiness/Public/SpatialReadinessSubsystem.h
@@ -27,6 +27,10 @@ public:
 	// overlapping volumes.
 	SPATIALREADINESS_API bool QueryReadiness(const FBox& Bounds, TArray<FString>& OutDescriptions, bool bAllDescriptions = false) const;
 
+#if !UE_BUILD_SHIPPING
+	SPATIALREADINESS_API FSpatialReadinessSimCallback* GetSimCallback() { return SimCallback; }
+#endif
+
 	/* begin: UWorldSubsystem */
 protected:
 	virtual bool ShouldCreateSubsystem(UObject* Outer) const override;
