Disable spatial readiness subsystem by default

diff --git a/Engine/Plugins/Experimental/SpatialReadiness/Source/SpatialReadiness/Private/SpatialReadiness/SpatialReadinessSubsystem.cpp b/Engine/Plugins/Experimental/SpatialReadiness/Source/SpatialReadiness/Private/SpatialReadiness/SpatialReadinessSubsystem.cpp
--- a/Engine/Plugins/Experimental/SpatialReadiness/Source/SpatialReadiness/Private/SpatialReadiness/SpatialReadinessSubsystem.cpp
+++ b/Engine/Plugins/Experimental/SpatialReadiness/Source/SpatialReadiness/Private/SpatialReadiness/SpatialReadinessSubsystem.cpp
@@ -12,7 +12,7 @@
 #include "Physics/Experimental/PhysScene_Chaos.h"
 #include "PhysicsProxy/SingleParticlePhysicsProxy.h"
 
-TAutoConsoleVariable<bool> CVarSpatialReadiness_OnDedicatedServer(TEXT("p.SpatialReadiness.OnDedicatedServer"), 1, TEXT("Enable spatial readiness on dedicated servers"), ECVF_Default);
+TAutoConsoleVariable<bool> CVarSpatialReadiness_OnDedicatedServer(TEXT("p.SpatialReadiness.OnDedicatedServer"), 0, TEXT("Enable spatial readiness on dedicated servers"), ECVF_Default);
 TAutoConsoleVariable<bool> CVarSpatialReadiness_OnClient(         TEXT("p.SpatialReadiness.OnClient"),          0, TEXT("Enable spatial readiness on clients"),           ECVF_Default);
 TAutoConsoleVariable<bool> CVarSpatialReadiness_Debug_NeverReady( TEXT("p.SpatialReadiness.Debug.NeverReady"),  0, TEXT("When true, unready volumes are never removed"),  ECVF_Default);
 
