* Freeze/refreeze around Integrate and Solve steps, so that midphase can occur on unfrozen bodies
* Update test unready volume so that it can be toggled on/off

#changelist validated
#virtualized

diff --git a/Engine/Plugins/Experimental/SpatialReadiness/Source/SpatialReadiness/Private/SpatialReadiness/SpatialReadinessSimCallback.cpp b/Engine/Plugins/Experimental/SpatialReadiness/Source/SpatialReadiness/Private/SpatialReadiness/SpatialReadinessSimCallback.cpp
--- a/Engine/Plugins/Experimental/SpatialReadiness/Source/SpatialReadiness/Private/SpatialReadiness/SpatialReadinessSimCallback.cpp
+++ b/Engine/Plugins/Experimental/SpatialReadiness/Source/SpatialReadiness/Private/SpatialReadiness/SpatialReadinessSimCallback.cpp
@@ -595,6 +595,16 @@ void FSpatialReadinessSimCallback::OnPostIntegrate_Internal()
 	UnFreezeParticles_PT();
 }
 
+void FSpatialReadinessSimCallback::OnPreSolve_Internal()
+{
+	FreezeParticles_PT();
+}
+
+void FSpatialReadinessSimCallback::OnPostSolve_Internal()
+{
+	UnFreezeParticles_PT();
+}
+
 uint32 FSpatialReadinessSimCallback::FHashMapTraits::GetElementID(const FUnreadyVolumeData_GT& Element)
 {
 	if (FSingleParticlePhysicsProxy* Proxy = Element.Proxy)
diff --git a/Engine/Plugins/Experimental/SpatialReadiness/Source/SpatialReadiness/Private/SpatialReadiness/SpatialReadinessVolumeComponent.cpp b/Engine/Plugins/Experimental/SpatialReadiness/Source/SpatialReadiness/Private/SpatialReadiness/SpatialReadinessVolumeComponent.cpp
--- a/Engine/Plugins/Experimental/SpatialReadiness/Source/SpatialReadiness/Private/SpatialReadiness/SpatialReadinessVolumeComponent.cpp
+++ b/Engine/Plugins/Experimental/SpatialReadiness/Source/SpatialReadiness/Private/SpatialReadiness/SpatialReadinessVolumeComponent.cpp
@@ -27,6 +27,11 @@ void USpatialReadinessVolumeComponent::EndPlay(const EEndPlayReason::Type EndPla
 	ReadinessVolume.MarkReady();
 }
 
+bool USpatialReadinessVolumeComponent::IsReady() const
+{
+	return ReadinessVolume.IsReady();
+}
+
 void USpatialReadinessVolumeComponent::MarkReady()
 {
 	ReadinessVolume.MarkReady();
diff --git a/Engine/Plugins/Experimental/SpatialReadiness/Source/SpatialReadiness/Public/SpatialReadinessSimCallback.h b/Engine/Plugins/Experimental/SpatialReadiness/Source/SpatialReadiness/Public/SpatialReadinessSimCallback.h
--- a/Engine/Plugins/Experimental/SpatialReadiness/Source/SpatialReadiness/Public/SpatialReadinessSimCallback.h
+++ b/Engine/Plugins/Experimental/SpatialReadiness/Source/SpatialReadiness/Public/SpatialReadinessSimCallback.h
@@ -52,7 +52,9 @@ struct FSpatialReadinessSimCallback
 		Chaos::ESimCallbackOptions::ParticleRegister |
 		Chaos::ESimCallbackOptions::MidPhaseModification |
 		Chaos::ESimCallbackOptions::PreIntegrate |
-		Chaos::ESimCallbackOptions::PostIntegrate>
+		Chaos::ESimCallbackOptions::PostIntegrate |
+		Chaos::ESimCallbackOptions::PreSolve |
+		Chaos::ESimCallbackOptions::PostSolve>
 {
 	using This = FSpatialReadinessSimCallback;
 
@@ -88,6 +90,8 @@ protected:
 	virtual void OnMidPhaseModification_Internal(Chaos::FMidPhaseModifierAccessor& Accessor) override;
 	virtual void OnPreIntegrate_Internal() override;
 	virtual void OnPostIntegrate_Internal() override;
+	virtual void OnPreSolve_Internal() override;
+	virtual void OnPostSolve_Internal() override;
 	/* end: TSimCallbackObject */
 
 	// Helpers
diff --git a/Engine/Plugins/Experimental/SpatialReadiness/Source/SpatialReadiness/Public/SpatialReadinessVolumeComponent.h b/Engine/Plugins/Experimental/SpatialReadiness/Source/SpatialReadiness/Public/SpatialReadinessVolumeComponent.h
--- a/Engine/Plugins/Experimental/SpatialReadiness/Source/SpatialReadiness/Public/SpatialReadinessVolumeComponent.h
+++ b/Engine/Plugins/Experimental/SpatialReadiness/Source/SpatialReadiness/Public/SpatialReadinessVolumeComponent.h
@@ -16,6 +16,9 @@ public:
 	virtual void BeginPlay() override;
 	virtual void EndPlay(const EEndPlayReason::Type EndPlayReason) override;
 
+	UFUNCTION(BlueprintCallable, Category = "SpatialReadiness")
+	bool IsReady() const;
+
 	UFUNCTION(BlueprintCallable, Category = "SpatialReadiness")
 	void MarkReady();
 
