Fix ensure in spatial readiness when setting bounds on readiness volume which doesn't have its internal delegates set yet

diff --git a/Engine/Plugins/Experimental/SpatialReadiness/Source/SpatialReadiness/Private/SpatialReadiness/SpatialReadinessVolume.cpp b/Engine/Plugins/Experimental/SpatialReadiness/Source/SpatialReadiness/Private/SpatialReadiness/SpatialReadinessVolume.cpp
--- a/Engine/Plugins/Experimental/SpatialReadiness/Source/SpatialReadiness/Private/SpatialReadiness/SpatialReadinessVolume.cpp
+++ b/Engine/Plugins/Experimental/SpatialReadiness/Source/SpatialReadiness/Private/SpatialReadiness/SpatialReadinessVolume.cpp
@@ -195,7 +195,7 @@ void FSpatialReadinessVolume::MarkUnready(const FBox& NewBounds, const FString&
 void FSpatialReadinessVolume::SetDescription(const FString& NewDescription)
 {
 #if WITH_SPATIAL_READINESS_DESCRIPTIONS
-	if (IsReady())
+	if ((!IsValid()) || IsReady())
 	{
 		Description = NewDescription;
 	}
@@ -208,7 +208,7 @@ void FSpatialReadinessVolume::SetDescription(const FString& NewDescription)
 
 void FSpatialReadinessVolume::SetBounds(const FBox& NewBounds)
 {
-	if (IsReady())
+	if ((!IsValid()) || IsReady())
 	{
 		Bounds = NewBounds;
 	}
