Remove default constructor from SpatialReadinessVolume, so that invalid volumes can't exist

diff --git a/Engine/Plugins/Experimental/SpatialReadiness/Source/SpatialReadiness/Private/SpatialReadiness/SpatialReadinessVolumeComponent.cpp b/Engine/Plugins/Experimental/SpatialReadiness/Source/SpatialReadiness/Private/SpatialReadiness/SpatialReadinessVolumeComponent.cpp
--- a/Engine/Plugins/Experimental/SpatialReadiness/Source/SpatialReadiness/Private/SpatialReadiness/SpatialReadinessVolumeComponent.cpp
+++ b/Engine/Plugins/Experimental/SpatialReadiness/Source/SpatialReadiness/Private/SpatialReadiness/SpatialReadinessVolumeComponent.cpp
@@ -24,22 +24,36 @@ void USpatialReadinessVolumeComponent::EndPlay(const EEndPlayReason::Type EndPla
 {
 	Super::EndPlay(EndPlayReason);
 
-	ReadinessVolume.MarkReady();
+	if (ReadinessVolume.IsSet())
+	{
+		ReadinessVolume->MarkReady();
+	}
 }
 
 bool USpatialReadinessVolumeComponent::IsReady() const
 {
-	return ReadinessVolume.IsReady();
+	if (ReadinessVolume.IsSet())
+	{
+		return ReadinessVolume->IsReady();
+	}
+
+	return false;
 }
 
 void USpatialReadinessVolumeComponent::MarkReady()
 {
-	ReadinessVolume.MarkReady();
+	if (ReadinessVolume.IsSet())
+	{
+		ReadinessVolume->MarkReady();
+	}
 }
 
 void USpatialReadinessVolumeComponent::MarkUnready()
 {
-	ReadinessVolume.MarkUnready();
+	if (ReadinessVolume.IsSet())
+	{
+		ReadinessVolume->MarkUnready();
+	}
 }
 
 void USpatialReadinessVolumeComponent::SetReadiness(bool bIsReady)
@@ -56,10 +70,16 @@ void USpatialReadinessVolumeComponent::SetReadiness(bool bIsReady)
 
 void USpatialReadinessVolumeComponent::SetDescription(const FString& InDescription)
 {
-	ReadinessVolume.SetDescription(InDescription);
+	if (ReadinessVolume.IsSet())
+	{
+		ReadinessVolume->SetDescription(InDescription);
+	}
 }
 
 void USpatialReadinessVolumeComponent::SetBounds(const FBox& InBounds)
 {
-	ReadinessVolume.SetBounds(InBounds);
+	if (ReadinessVolume.IsSet())
+	{
+		ReadinessVolume->SetBounds(InBounds);
+	}
 }
diff --git a/Engine/Plugins/Experimental/SpatialReadiness/Source/SpatialReadiness/Public/SpatialReadinessVolume.h b/Engine/Plugins/Experimental/SpatialReadiness/Source/SpatialReadiness/Public/SpatialReadinessVolume.h
--- a/Engine/Plugins/Experimental/SpatialReadiness/Source/SpatialReadiness/Public/SpatialReadinessVolume.h
+++ b/Engine/Plugins/Experimental/SpatialReadiness/Source/SpatialReadiness/Public/SpatialReadinessVolume.h
@@ -53,9 +53,6 @@ public:
 	// When a handle is destroyed, it will make the corresponding volume as ready
 	SPATIALREADINESS_API ~FSpatialReadinessVolume();
 
-	// The default constructor will create an invalid handle
-	FSpatialReadinessVolume() = default;
-
 	// Readiness handles are unique and cannot be copied
 	FSpatialReadinessVolume(const FSpatialReadinessVolume&) = delete;
 	FSpatialReadinessVolume& operator=(const FSpatialReadinessVolume&) = delete;
diff --git a/Engine/Plugins/Experimental/SpatialReadiness/Source/SpatialReadiness/Public/SpatialReadinessVolumeComponent.h b/Engine/Plugins/Experimental/SpatialReadiness/Source/SpatialReadiness/Public/SpatialReadinessVolumeComponent.h
--- a/Engine/Plugins/Experimental/SpatialReadiness/Source/SpatialReadiness/Public/SpatialReadinessVolumeComponent.h
+++ b/Engine/Plugins/Experimental/SpatialReadiness/Source/SpatialReadiness/Public/SpatialReadinessVolumeComponent.h
@@ -45,6 +45,6 @@ public:
 
 protected:
 
-	FSpatialReadinessVolume ReadinessVolume;
+	TOptional<FSpatialReadinessVolume> ReadinessVolume;
 };
 
