SpatialReadiness log command updated to include list of every frozen particle and their bounds

diff --git a/Engine/Plugins/Experimental/SpatialReadiness/Source/SpatialReadiness/Private/SpatialReadiness/SpatialReadinessDebug.cpp b/Engine/Plugins/Experimental/SpatialReadiness/Source/SpatialReadiness/Private/SpatialReadiness/SpatialReadinessDebug.cpp
--- a/Engine/Plugins/Experimental/SpatialReadiness/Source/SpatialReadiness/Private/SpatialReadiness/SpatialReadinessDebug.cpp
+++ b/Engine/Plugins/Experimental/SpatialReadiness/Source/SpatialReadiness/Private/SpatialReadiness/SpatialReadinessDebug.cpp
@@ -5,6 +5,7 @@
 #include "SpatialReadinessSubsystem.h"
 #include "SpatialReadinessSimCallback.h"
 #include "PhysicsProxy/SingleParticlePhysicsProxy.h"
+#include "Physics/Experimental/PhysScene_Chaos.h"
 #include "Chaos/ParticleHandle.h"
 #include "Engine/World.h"
 
@@ -21,7 +22,7 @@ TAutoConsoleVariable<bool> CVarSpatialReadinessDebugDraw(
 
 static FAutoConsoleCommandWithWorld CmdSpatialReadinessPrint(
 	TEXT("p.SpatialReadiness.Print"),
-	TEXT("Command to log all current unready volumes and their statuses"),
+	TEXT("Command to log all current unready volumes and frozen rigid particles"),
 	FConsoleCommandWithWorldDelegate::CreateStatic([](UWorld* World)
 	{
 #if !UE_BUILD_SHIPPING
@@ -31,6 +32,20 @@ static FAutoConsoleCommandWithWorld CmdSpatialReadinessPrint(
 			return;
 		}
 
+		FPhysScene_Chaos* Scene = static_cast<FPhysScene_Chaos*>(World->GetPhysicsScene());
+		if (Scene == nullptr)
+		{
+			UE_LOG(LogSpatialReadiness, Log, TEXT("Can't print spatial readiness: no physics scene"));
+			return;
+		}
+
+		Chaos::FPhysicsSolver* Solver = Scene->GetSolver();
+		if (Solver == nullptr)
+		{
+			UE_LOG(LogSpatialReadiness, Log, TEXT("Can't print spatial readiness: no solver"));
+			return;
+		}
+
 		USpatialReadiness* SpatialReadiness = World->GetSubsystem<USpatialReadiness>();
 		if (SpatialReadiness == nullptr)
 		{
@@ -47,7 +62,7 @@ static FAutoConsoleCommandWithWorld CmdSpatialReadinessPrint(
 
 		int32 VolumeNum = 0;
 		TStringBuilder<(1<<14)> Builder;
-		Builder << TEXT("Unready Volumes (" << SimCallback->GetNumUnreadyVolumes_GT() << "):\n");
+		Builder << "Unready Volumes (" << SimCallback->GetNumUnreadyVolumes_GT() << "):\n";
 		SimCallback->ForEachVolumeData_GT([&Builder, &VolumeNum](const FUnreadyVolumeData_GT& Data)
 		{
 			FSingleParticlePhysicsProxy* Proxy = Data.Proxy;
@@ -72,6 +87,19 @@ static FAutoConsoleCommandWithWorld CmdSpatialReadinessPrint(
 			Builder << " " << (VolumeNum++) << ". " << Data.Description << " : " << *Bounds.ToString() << "\n";
 		});
 
+		Solver->EnqueueCommandImmediate([&Builder, SimCallback]()
+		{
+			int32 ParticleNum = 0;
+			Builder << "\nFrozen Particles (" << SimCallback->GetNumUnreadyRigidParticles_PT() << "):\n";
+			SimCallback->ForEachUnreadyRigidParticle_PT([&Builder, &ParticleNum](const Chaos::FPBDRigidParticleHandle* RigidParticle) -> bool
+			{
+				const FAABB3 Bounds = RigidParticle->WorldSpaceInflatedBounds();
+				const FName DebugName = *RigidParticle->GetDebugName();
+				Builder << " " << (ParticleNum++) << ". " << *DebugName.ToString() << " : " << *Bounds.ToString() << "\n";
+				return true;
+			});
+		});
+
 		UE_LOG(LogSpatialReadiness, Log, TEXT("%s"), *Builder);
 #endif
 	}));
diff --git a/Engine/Plugins/Experimental/SpatialReadiness/Source/SpatialReadiness/Private/SpatialReadiness/SpatialReadinessSimCallback.cpp b/Engine/Plugins/Experimental/SpatialReadiness/Source/SpatialReadiness/Private/SpatialReadiness/SpatialReadinessSimCallback.cpp
--- a/Engine/Plugins/Experimental/SpatialReadiness/Source/SpatialReadiness/Private/SpatialReadiness/SpatialReadinessSimCallback.cpp
+++ b/Engine/Plugins/Experimental/SpatialReadiness/Source/SpatialReadiness/Private/SpatialReadiness/SpatialReadinessSimCallback.cpp
@@ -641,3 +641,11 @@ void FSpatialReadinessSimCallback::ForEachUnreadyRigidParticle_PT(const TFunctio
 	}
 }
 
+int32 FSpatialReadinessSimCallback::GetNumUnreadyRigidParticles_PT() const
+{
+	return
+		NewRegistered_UnreadyRigidParticles_PT.Num() +
+		OldRegistered_UnreadyRigidParticles_PT.Num() +
+		MidPhase_UnreadyRigidParticles_PT.Num();
+}
+
diff --git a/Engine/Plugins/Experimental/SpatialReadiness/Source/SpatialReadiness/Public/SpatialReadinessSimCallback.h b/Engine/Plugins/Experimental/SpatialReadiness/Source/SpatialReadiness/Public/SpatialReadinessSimCallback.h
--- a/Engine/Plugins/Experimental/SpatialReadiness/Source/SpatialReadiness/Public/SpatialReadinessSimCallback.h
+++ b/Engine/Plugins/Experimental/SpatialReadiness/Source/SpatialReadiness/Public/SpatialReadinessSimCallback.h
@@ -83,6 +83,12 @@ public:
 	// Get the number of unready volumes
 	int32 GetNumUnreadyVolumes_GT() const;
 
+	// Perform an operation on each of the currently frozen particles - only valid on the physics thread
+	void ForEachUnreadyRigidParticle_PT(const TFunction<bool(Chaos::FPBDRigidParticleHandle*)>& Lambda) const;
+
+	// Get the total number of currently frozen particles - only valid on the physics thread
+	int32 GetNumUnreadyRigidParticles_PT() const;
+
 	/* begin: TSimCallbackObject */
 protected:
 	virtual void OnPreSimulate_Internal() override;
@@ -143,7 +149,6 @@ protected:
 	TSet<Chaos::FPBDRigidParticleHandle*> NewRegistered_UnreadyRigidParticles_PT;
 	TSet<Chaos::FPBDRigidParticleHandle*> OldRegistered_UnreadyRigidParticles_PT;
 	TSet<Chaos::FPBDRigidParticleHandle*> MidPhase_UnreadyRigidParticles_PT;
-	void ForEachUnreadyRigidParticle_PT(const TFunction<bool(Chaos::FPBDRigidParticleHandle*)>& Lambda) const;
 
 	// "Unready" particles are forced to be stationary in PreSimulate, and restored
 	// to their previous state in PostIntegrate. The values needed for restoration
