Add more cycle stats to spatial readiness to isolate cost of debug code

diff --git a/Engine/Plugins/Experimental/SpatialReadiness/Source/SpatialReadiness/Private/SpatialReadiness/SpatialReadinessSimCallback.cpp b/Engine/Plugins/Experimental/SpatialReadiness/Source/SpatialReadiness/Private/SpatialReadiness/SpatialReadinessSimCallback.cpp
--- a/Engine/Plugins/Experimental/SpatialReadiness/Source/SpatialReadiness/Private/SpatialReadiness/SpatialReadinessSimCallback.cpp
+++ b/Engine/Plugins/Experimental/SpatialReadiness/Source/SpatialReadiness/Private/SpatialReadiness/SpatialReadinessSimCallback.cpp
@@ -145,6 +145,7 @@ int32 FSpatialReadinessSimCallback::AddUnreadyVolume_GT(const FBox& Bounds, cons
 #if !UE_BUILD_SHIPPING
 	{
 		// Draw a box in CVD representing the unready volume
+		SCOPE_CYCLE_COUNTER(STAT_SpatialReadiness_Physics_ChaosVisualDebugger)
 		const FName Tag = *FString::Printf(TEXT("Created: %s"), *Description);
 		const FLinearColor Color = FLinearColor::Green;
 		CVD_TRACE_DEBUG_DRAW_BOX(FBox(Bounds.Min, Bounds.Max), Tag, Color.ToFColorSRGB(), FChaosVisualDebuggerTrace::GetSolverID(*GetSolver()));
@@ -166,6 +167,7 @@ void FSpatialReadinessSimCallback::RemoveUnreadyVolume_GT(int32 UnreadyVolumeInd
 
 #if !UE_BUILD_SHIPPING
 	{
+		SCOPE_CYCLE_COUNTER(STAT_SpatialReadiness_Physics_ChaosVisualDebugger)
 #if WITH_SPATIAL_READINESS_DESCRIPTIONS
 		const FName Tag = *FString::Printf(TEXT("Destroyed: %s"), *VolumeData->Description);
 #else
@@ -408,6 +410,7 @@ void FSpatialReadinessSimCallback::FreezeParticles_PT()
 		if (CVarSpatialReadinessDebugDraw.GetValueOnAnyThread())
 		{
 			// Draw debug box around the frozen object
+			SCOPE_CYCLE_COUNTER(STAT_SpatialReadiness_Physics_DebugDraw)
 			const FAABB3 Bounds = RigidParticle->WorldSpaceInflatedBounds();
 			Chaos::FDebugDrawQueue::GetInstance().DrawDebugBox(Bounds.Center(), Bounds.Extents() * .5f, FQuat::Identity, FColor::Yellow, false, -1.f, -1, 0.f);
 		}
@@ -416,6 +419,7 @@ void FSpatialReadinessSimCallback::FreezeParticles_PT()
 #if !UE_BUILD_SHIPPING
 		{
 			// Draw a box in CVD representing the frozen object
+			SCOPE_CYCLE_COUNTER(STAT_SpatialReadiness_Physics_ChaosVisualDebugger)
 			const FName Tag = *RigidParticle->GetDebugName();
 			const FAABB3 Bounds = RigidParticle->WorldSpaceInflatedBounds();
 			CVD_TRACE_DEBUG_DRAW_BOX(FBox(Bounds.Min(), Bounds.Max()), Tag, FColor::Yellow, FChaosVisualDebuggerTrace::GetSolverID(*GetSolver()));
@@ -537,6 +541,7 @@ void FSpatialReadinessSimCallback::OnMidPhaseModification_Internal(FMidPhaseModi
 #if ENABLE_DRAW_DEBUG
 		if (CVarSpatialReadinessDebugDraw.GetValueOnAnyThread())
 		{
+			SCOPE_CYCLE_COUNTER(STAT_SpatialReadiness_Physics_DebugDraw)
 			const FAABB3 Bounds = UnreadyVolume->WorldSpaceInflatedBounds();
 			Chaos::FDebugDrawQueue::GetInstance().DrawDebugBox(Bounds.Center(), Bounds.Extents() * .5f, FQuat::Identity, FColor::Red, false, -1.f, -1, 0.f);
 		}
@@ -544,6 +549,7 @@ void FSpatialReadinessSimCallback::OnMidPhaseModification_Internal(FMidPhaseModi
 
 #if !UE_BUILD_SHIPPING
 		{
+			SCOPE_CYCLE_COUNTER(STAT_SpatialReadiness_Physics_ChaosVisualDebugger)
 			// Draw a box in CVD representing the unready volume
 			const FName Tag = *UnreadyVolume->GetDebugName();
 			const FLinearColor Color = FLinearColor(1.f, .25f, .5f);
@@ -555,6 +561,8 @@ void FSpatialReadinessSimCallback::OnMidPhaseModification_Internal(FMidPhaseModi
 		// Go through every mid-phase which involves this volume
 		for (FMidPhaseModifier& MidPhase : Accessor.GetMidPhases(UnreadyVolume))
 		{
+			SCOPE_CYCLE_COUNTER(STAT_SpatialReadiness_Physics_MidPhase_AccessorLoop)
+
 			// Get the particle that is not the unready volume
 			FGeometryParticleHandle* GeometryParticle = MidPhase.GetOtherParticle(UnreadyVolume);
 			if (GeometryParticle == nullptr)
@@ -571,6 +579,7 @@ void FSpatialReadinessSimCallback::OnMidPhaseModification_Internal(FMidPhaseModi
 #if ENABLE_DRAW_DEBUG
 			if (CVarSpatialReadinessDebugDraw.GetValueOnAnyThread())
 			{
+				SCOPE_CYCLE_COUNTER(STAT_SpatialReadiness_Physics_DebugDraw)
 				const FAABB3 Bounds = RigidParticle->WorldSpaceInflatedBounds();
 				Chaos::FDebugDrawQueue::GetInstance().DrawDebugBox(Bounds.Center(), Bounds.Extents() * .5f, FQuat::Identity, FColor::Yellow, false, -1.f, -1, 0.f);
 			}
diff --git a/Engine/Plugins/Experimental/SpatialReadiness/Source/SpatialReadiness/Public/SpatialReadinessStats.h b/Engine/Plugins/Experimental/SpatialReadiness/Source/SpatialReadiness/Public/SpatialReadinessStats.h
--- a/Engine/Plugins/Experimental/SpatialReadiness/Source/SpatialReadiness/Public/SpatialReadinessStats.h
+++ b/Engine/Plugins/Experimental/SpatialReadiness/Source/SpatialReadiness/Public/SpatialReadinessStats.h
@@ -14,4 +14,7 @@ DECLARE_CYCLE_STAT(TEXT("SpatialReadiness_Physics_UnFreezeParticlesPT"), STAT_Sp
 DECLARE_CYCLE_STAT(TEXT("SpatialReadiness_Physics_PreSimulate"), STAT_SpatialReadiness_Physics_PreSimulate, STATGROUP_SpatialReadiness);
 DECLARE_CYCLE_STAT(TEXT("SpatialReadiness_Physics_ParticlesRegistered"), STAT_SpatialReadiness_Physics_ParticlesRegistered, STATGROUP_SpatialReadiness);
 DECLARE_CYCLE_STAT(TEXT("SpatialReadiness_Physics_MidPhase"), STAT_SpatialReadiness_Physics_MidPhase, STATGROUP_SpatialReadiness);
+DECLARE_CYCLE_STAT(TEXT("SpatialReadiness_Physics_MidPhase_AccessorLoop"), STAT_SpatialReadiness_Physics_MidPhase_AccessorLoop, STATGROUP_SpatialReadiness);
+DECLARE_CYCLE_STAT(TEXT("SpatialReadiness_Physics_ChaosVisualDebugger"), STAT_SpatialReadiness_Physics_ChaosVisualDebugger, STATGROUP_SpatialReadiness);
+DECLARE_CYCLE_STAT(TEXT("SpatialReadiness_Physics_DebugDraw"), STAT_SpatialReadiness_Physics_DebugDraw, STATGROUP_SpatialReadiness);
 
