Sort list of frozen bodies in spatial readiness debug print

diff --git a/Engine/Plugins/Experimental/SpatialReadiness/Source/SpatialReadiness/Private/SpatialReadiness/SpatialReadinessDebug.cpp b/Engine/Plugins/Experimental/SpatialReadiness/Source/SpatialReadiness/Private/SpatialReadiness/SpatialReadinessDebug.cpp
--- a/Engine/Plugins/Experimental/SpatialReadiness/Source/SpatialReadiness/Private/SpatialReadiness/SpatialReadinessDebug.cpp
+++ b/Engine/Plugins/Experimental/SpatialReadiness/Source/SpatialReadiness/Private/SpatialReadiness/SpatialReadinessDebug.cpp
@@ -20,10 +20,10 @@ TAutoConsoleVariable<bool> CVarSpatialReadinessDebugDraw(
 	ECVF_Default);
 
 
-static FAutoConsoleCommandWithWorld CmdSpatialReadinessPrint(
+static FAutoConsoleCommandWithWorldAndArgs CmdSpatialReadinessPrint(
 	TEXT("p.SpatialReadiness.Print"),
 	TEXT("Command to log all current unready volumes and frozen rigid particles"),
-	FConsoleCommandWithWorldDelegate::CreateStatic([](UWorld* World)
+	FConsoleCommandWithWorldAndArgsDelegate::CreateLambda([](const TArray<FString>& Args, UWorld* World)
 	{
 #if !UE_BUILD_SHIPPING
 		if (World == nullptr)
@@ -62,7 +62,7 @@ static FAutoConsoleCommandWithWorld CmdSpatialReadinessPrint(
 
 		int32 VolumeNum = 0;
 		TStringBuilder<(1<<14)> Builder;
-		Builder << "Unready Volumes (" << SimCallback->GetNumUnreadyVolumes_GT() << "):\n";
+		Builder << "\nUnready Volumes (" << SimCallback->GetNumUnreadyVolumes_GT() << "):\n";
 		SimCallback->ForEachVolumeData_GT([&Builder, &VolumeNum](const FUnreadyVolumeData_GT& Data)
 		{
 			FSingleParticlePhysicsProxy* Proxy = Data.Proxy;
@@ -87,17 +87,55 @@ static FAutoConsoleCommandWithWorld CmdSpatialReadinessPrint(
 			Builder << " " << (VolumeNum++) << ". " << Data.Description << " : " << *Bounds.ToString() << "\n";
 		});
 
-		Solver->EnqueueCommandImmediate([&Builder, SimCallback]()
+		// Check to see if we're required to spatially sort frozen bodies
+		bool bSortFrozen = false;
+		for (const FString& Str : Args)
+		{
+			if (Str.Compare("-s", ESearchCase::IgnoreCase) == 0)
+			{
+				bSortFrozen = true;
+				break;
+			}
+		}
+
+		// Make a lambda that will accumulate a string for printing frozen particles
+		int32 ParticleNum = 0;
+		const auto PrintParticle = [&Builder, &ParticleNum](const Chaos::FPBDRigidParticleHandle* RigidParticle) -> bool
+		{
+			const FAABB3 Bounds = RigidParticle->WorldSpaceInflatedBounds();
+			const FName DebugName = *RigidParticle->GetDebugName();
+			Builder << " " << (ParticleNum++) << ". " << *DebugName.ToString() << " : " << *Bounds.ToString() << "\n";
+			return true;
+		};
+
+		// Enqueue a command to print all particles
+		Solver->EnqueueCommandImmediate([&Builder, SimCallback, &PrintParticle, bSortFrozen]()
 		{
 			int32 ParticleNum = 0;
 			Builder << "\nFrozen Particles (" << SimCallback->GetNumUnreadyRigidParticles_PT() << "):\n";
-			SimCallback->ForEachUnreadyRigidParticle_PT([&Builder, &ParticleNum](const Chaos::FPBDRigidParticleHandle* RigidParticle) -> bool
+
+			if (bSortFrozen)
 			{
-				const FAABB3 Bounds = RigidParticle->WorldSpaceInflatedBounds();
-				const FName DebugName = *RigidParticle->GetDebugName();
-				Builder << " " << (ParticleNum++) << ". " << *DebugName.ToString() << " : " << *Bounds.ToString() << "\n";
-				return true;
-			});
+				TArray<const Chaos::FPBDRigidParticleHandle*> Particles;
+				Particles.Reserve(SimCallback->GetNumUnreadyRigidParticles_PT());
+				SimCallback->ForEachUnreadyRigidParticle_PT([&Particles](const Chaos::FPBDRigidParticleHandle* RigidParticle) -> bool
+				{
+					Particles.Add(RigidParticle);
+					return true;
+				});
+				Particles.Sort([](const Chaos::FPBDRigidParticleHandle& A, const Chaos::FPBDRigidParticleHandle& B) -> bool
+				{
+					return !(B.GetX().X > A.GetX().X || B.GetX().Y > A.GetX().Y || B.GetX().Z > A.GetX().Z);
+				});
+				for (const Chaos::FPBDRigidParticleHandle* Particle : Particles)
+				{
+					PrintParticle(Particle);
+				}
+			}
+			else
+			{
+				SimCallback->ForEachUnreadyRigidParticle_PT(PrintParticle);
+			}
 		});
 
 		UE_LOG(LogSpatialReadiness, Log, TEXT("%s"), *Builder);
