* Allow spatial readiness to be enabled in Standalone mode, for testing in the physics sandbox
* Fix spatial readiness debug print for frozen particles by fixing threading issue

diff --git a/Engine/Plugins/Experimental/SpatialReadiness/Source/SpatialReadiness/Private/SpatialReadiness/SpatialReadinessDebug.cpp b/Engine/Plugins/Experimental/SpatialReadiness/Source/SpatialReadiness/Private/SpatialReadiness/SpatialReadinessDebug.cpp
--- a/Engine/Plugins/Experimental/SpatialReadiness/Source/SpatialReadiness/Private/SpatialReadiness/SpatialReadinessDebug.cpp
+++ b/Engine/Plugins/Experimental/SpatialReadiness/Source/SpatialReadiness/Private/SpatialReadiness/SpatialReadinessDebug.cpp
@@ -60,84 +60,105 @@ static FAutoConsoleCommandWithWorldAndArgs CmdSpatialReadinessPrint(
 			return;
 		}
 
-		int32 VolumeNum = 0;
-		TStringBuilder<(1<<14)> Builder;
-		Builder << "\nUnready Volumes (" << SimCallback->GetNumUnreadyVolumes_GT() << "):\n";
-		SimCallback->ForEachVolumeData_GT([&Builder, &VolumeNum](const FUnreadyVolumeData_GT& Data)
 		{
-			FSingleParticlePhysicsProxy* Proxy = Data.Proxy;
-			if (Proxy == nullptr)
+			// Print unready volumes string
+			int32 VolumeNum = 0;
+			TStringBuilder<(1<<14)> Builder;
+			Builder << "\nUnready Volumes (" << SimCallback->GetNumUnreadyVolumes_GT() << "):\n";
+			SimCallback->ForEachVolumeData_GT([&Builder, &VolumeNum](const FUnreadyVolumeData_GT& Data)
 			{
-				return;
-			}
+				FSingleParticlePhysicsProxy* Proxy = Data.Proxy;
+				if (Proxy == nullptr)
+				{
+					return;
+				}
 
-			if (Proxy->GetMarkedDeleted())
-			{
-				return;
-			}
+				if (Proxy->GetMarkedDeleted())
+				{
+					return;
+				}
 
-			FGeometryParticleHandle* UnreadyVolume = Proxy->GetHandle_LowLevel();
-			if (UnreadyVolume == nullptr)
-			{
-				return;
-			}
+				FGeometryParticleHandle* UnreadyVolume = Proxy->GetHandle_LowLevel();
+				if (UnreadyVolume == nullptr)
+				{
+					return;
+				}
 
-			const FAABB3 Bounds = UnreadyVolume->WorldSpaceInflatedBounds();
+				const FAABB3 Bounds = UnreadyVolume->WorldSpaceInflatedBounds();
 
-			Builder << " " << (VolumeNum++) << ". " << Data.Description << " : " << *Bounds.ToString() << "\n";
-		});
+				Builder << " " << (VolumeNum++) << ". " << Data.Description << " : " << *Bounds.ToString() << "\n";
+			});
 
-		// Check to see if we're required to spatially sort frozen bodies
-		bool bSortFrozen = false;
-		for (const FString& Str : Args)
-		{
-			if (Str.Compare("-s", ESearchCase::IgnoreCase) == 0)
-			{
-				bSortFrozen = true;
-				break;
-			}
+			// Print the string
+			UE_LOG(LogSpatialReadiness, Log, TEXT("%s"), *Builder);
 		}
 
-		// Make a lambda that will accumulate a string for printing frozen particles
-		int32 ParticleNum = 0;
-		const auto PrintParticle = [&Builder, &ParticleNum](const Chaos::FPBDRigidParticleHandle* RigidParticle) -> bool
 		{
-			const FAABB3 Bounds = RigidParticle->WorldSpaceInflatedBounds();
-			const FName DebugName = *RigidParticle->GetDebugName();
-			Builder << " " << (ParticleNum++) << ". " << *DebugName.ToString() << " : " << *Bounds.ToString() << "\n";
-			return true;
-		};
-
-		// Enqueue a command to print all particles
-		Solver->EnqueueCommandImmediate([&Builder, SimCallback, &PrintParticle, bSortFrozen]()
-		{
-			int32 ParticleNum = 0;
-			Builder << "\nFrozen Particles (" << SimCallback->GetNumUnreadyRigidParticles_PT() << "):\n";
+			// Print frozen particles string
 
-			if (bSortFrozen)
+			// Check to see if we're required to spatially sort frozen bodies
+			bool bSortFrozen = false;
+			for (const FString& Str : Args)
 			{
-				TArray<const Chaos::FPBDRigidParticleHandle*> Particles;
-				Particles.Reserve(SimCallback->GetNumUnreadyRigidParticles_PT());
-				SimCallback->ForEachUnreadyRigidParticle_PT([&Particles](const Chaos::FPBDRigidParticleHandle* RigidParticle) -> bool
+				if (Str.Compare("-s", ESearchCase::IgnoreCase) == 0)
+				{
+					bSortFrozen = true;
+					break;
+				}
+			}
+
+			// Enqueue a command to print all particles
+			Solver->EnqueueCommandImmediate([SimCallback, bSortFrozen]()
+			{
+				TStringBuilder<(1<<14)> Builder;
+				Builder << "\nFrozen Particles (" << SimCallback->GetNumUnreadyRigidParticles_PT() << "):\n";
+
+				// Make a lambda that will accumulate a string for printing frozen particles
+				const auto PrintParticle = [&Builder](const Chaos::FPBDRigidParticleHandle* RigidParticle) -> bool
 				{
-					Particles.Add(RigidParticle);
+					static int32 ParticleNum = 0;
+					const FAABB3 Bounds = RigidParticle->WorldSpaceInflatedBounds();
+					const FVector Pos = RigidParticle->GetX();
+					const FName DebugName = *RigidParticle->GetDebugName();
+					Builder << " " << (ParticleNum++) << ". " << *DebugName.ToString() << " : " << *Pos.ToString() << "\n";
 					return true;
-				});
-				Particles.Sort([](const Chaos::FPBDRigidParticleHandle& A, const Chaos::FPBDRigidParticleHandle& B) -> bool
+				};
+
+				if (bSortFrozen)
 				{
-					return !(B.GetX().X > A.GetX().X || B.GetX().Y > A.GetX().Y || B.GetX().Z > A.GetX().Z);
-				});
-				for (const Chaos::FPBDRigidParticleHandle* Particle : Particles)
+					TArray<const Chaos::FPBDRigidParticleHandle*> Particles;
+					Particles.Reserve(SimCallback->GetNumUnreadyRigidParticles_PT());
+					SimCallback->ForEachUnreadyRigidParticle_PT([&Particles](const Chaos::FPBDRigidParticleHandle* RigidParticle) -> bool
+					{
+						Particles.Add(RigidParticle);
+						return true;
+					});
+					Particles.Sort([](const Chaos::FPBDRigidParticleHandle& A, const Chaos::FPBDRigidParticleHandle& B) -> bool
+					{
+						const FVector AX = A.GetX();
+						const FVector BX = B.GetX();
+						return
+							(AX.X < BX.X) || (
+								(AX.X == BX.X) && (
+									(AX.Y < BX.Y) || (
+										(AX.Y == BX.Y) && (AX.Z < BX.Z)
+									)
+								)
+							);
+					});
+					for (const Chaos::FPBDRigidParticleHandle* Particle : Particles)
+					{
+						PrintParticle(Particle);
+					}
+				}
+				else
 				{
-					PrintParticle(Particle);
+					SimCallback->ForEachUnreadyRigidParticle_PT(PrintParticle);
 				}
-			}
-			else
-			{
-				SimCallback->ForEachUnreadyRigidParticle_PT(PrintParticle);
-			}
-		});
 
-		UE_LOG(LogSpatialReadiness, Log, TEXT("%s"), *Builder);
+				// Print the string
+				UE_LOG(LogSpatialReadiness, Log, TEXT("%s"), *Builder);
+			});
+		}
 #endif
 	}));
diff --git a/Engine/Plugins/Experimental/SpatialReadiness/Source/SpatialReadiness/Private/SpatialReadiness/SpatialReadinessSubsystem.cpp b/Engine/Plugins/Experimental/SpatialReadiness/Source/SpatialReadiness/Private/SpatialReadiness/SpatialReadinessSubsystem.cpp
--- a/Engine/Plugins/Experimental/SpatialReadiness/Source/SpatialReadiness/Private/SpatialReadiness/SpatialReadinessSubsystem.cpp
+++ b/Engine/Plugins/Experimental/SpatialReadiness/Source/SpatialReadiness/Private/SpatialReadiness/SpatialReadinessSubsystem.cpp
@@ -15,9 +15,10 @@
 
 #include UE_INLINE_GENERATED_CPP_BY_NAME(SpatialReadinessSubsystem)
 
-TAutoConsoleVariable<bool> CVarSpatialReadiness_OnDedicatedServer(TEXT("p.SpatialReadiness.OnDedicatedServer"), 0, TEXT("Enable spatial readiness on dedicated servers"), ECVF_Default);
-TAutoConsoleVariable<bool> CVarSpatialReadiness_OnClient(         TEXT("p.SpatialReadiness.OnClient"),          0, TEXT("Enable spatial readiness on clients"),           ECVF_Default);
-TAutoConsoleVariable<bool> CVarSpatialReadiness_Debug_NeverReady( TEXT("p.SpatialReadiness.Debug.NeverReady"),  0, TEXT("When true, unready volumes are never removed"),  ECVF_Default);
+TAutoConsoleVariable<bool> CVarSpatialReadiness_OnDedicatedServer(TEXT("p.SpatialReadiness.OnDedicatedServer"), 0, TEXT("Enable spatial readiness on dedicated servers"),  ECVF_Default);
+TAutoConsoleVariable<bool> CVarSpatialReadiness_OnStandalone(     TEXT("p.SpatialReadiness.OnStandalone"),      0, TEXT("Enable spatial readiness on standalone servers"), ECVF_Default);
+TAutoConsoleVariable<bool> CVarSpatialReadiness_OnClient(         TEXT("p.SpatialReadiness.OnClient"),          0, TEXT("Enable spatial readiness on clients"),            ECVF_Default);
+TAutoConsoleVariable<bool> CVarSpatialReadiness_Debug_NeverReady( TEXT("p.SpatialReadiness.Debug.NeverReady"),  0, TEXT("When true, unready volumes are never removed"),   ECVF_Default);
 
 using namespace Chaos;
 
@@ -45,6 +46,10 @@ bool USpatialReadiness::StaticShouldCreateSubsystem(UObject* Outer)
 		{
 			return true;
 		}
+		else if (NetMode == NM_Standalone && CVarSpatialReadiness_OnStandalone.GetValueOnAnyThread())
+		{
+			return true;
+		}
 	}
 	return false;
 }
