Use a decrementing counter to keep unready rigids around for one extra frame after losing contact with unready volumes. Fixes the case where unready volumes are deleted on the same frame that the collision which they were protecting is created, since collision takes one extra sim tick to actually appear.

diff --git a/Engine/Plugins/Experimental/SpatialReadiness/Source/SpatialReadiness/Private/SpatialReadiness/SpatialReadinessSimCallback.cpp b/Engine/Plugins/Experimental/SpatialReadiness/Source/SpatialReadiness/Private/SpatialReadiness/SpatialReadinessSimCallback.cpp
--- a/Engine/Plugins/Experimental/SpatialReadiness/Source/SpatialReadiness/Private/SpatialReadiness/SpatialReadinessSimCallback.cpp
+++ b/Engine/Plugins/Experimental/SpatialReadiness/Source/SpatialReadiness/Private/SpatialReadiness/SpatialReadinessSimCallback.cpp
@@ -439,6 +439,20 @@ void FSpatialReadinessSimCallback::UnFreezeParticles_PT()
 	ParticleDataCache_PT.Reset();
 }
 
+void FSpatialReadinessSimCallback::DecrementUnreadyRigidParticles_PT()
+{
+	SCOPE_CYCLE_COUNTER(STAT_SpatialReadiness_Physics_DecrementUnreadyRigidParticlesPT)
+
+	// Decrement each counter, 
+	for (auto It = UnreadyRigidParticles_PT.CreateIterator(); It; ++It)
+	{
+		if (--It->Value < 0)
+		{
+			It.RemoveCurrent();
+		}
+	}
+}
+
 void FSpatialReadinessSimCallback::OnPreSimulate_Internal()
 {
 	SCOPE_CYCLE_COUNTER(STAT_SpatialReadiness_Physics_PreSimulate)
@@ -458,18 +472,14 @@ void FSpatialReadinessSimCallback::OnPreSimulate_Internal()
 			UnreadyVolumeParticles_PT.Remove(ParticleProxy);
 		}
 	}
+
+	DecrementUnreadyRigidParticles_PT();
 }
 
 void FSpatialReadinessSimCallback::OnParticlesRegistered_Internal(TArray<FSingleParticlePhysicsProxy*>& RegisteredProxies)
 {
 	SCOPE_CYCLE_COUNTER(STAT_SpatialReadiness_Physics_ParticlesRegistered)
 
-	// Empty the list of newly registered particles which are unready,
-	// so that we don't double count them
-	//Registered_UnreadyRigidParticles_PT.Reset();
-	Swap(NewRegistered_UnreadyRigidParticles_PT, OldRegistered_UnreadyRigidParticles_PT);
-	NewRegistered_UnreadyRigidParticles_PT.Reset();
-
 	// For each newly added particle, query against unready volumes to see if
 	// it should be frozen.
 	for (FSingleParticlePhysicsProxy* ParticleProxy : RegisteredProxies)
@@ -493,7 +503,7 @@ void FSpatialReadinessSimCallback::OnParticlesRegistered_Internal(TArray<FSingle
 		{
 			// If this volume hit an unready volume, query readiness returns false
 			// and we must mark the particle as frozen
-			NewRegistered_UnreadyRigidParticles_PT.Add(RigidParticle);
+			UnreadyRigidParticles_PT.FindOrAdd(RigidParticle) = 2;
 		}
 	}
 }
@@ -502,9 +512,6 @@ void FSpatialReadinessSimCallback::OnMidPhaseModification_Internal(FMidPhaseModi
 {
 	SCOPE_CYCLE_COUNTER(STAT_SpatialReadiness_Physics_MidPhase)
 
-	// Clear the list of unready particles
-	MidPhase_UnreadyRigidParticles_PT.Reset();
-
 	// Go through every unready volume
 	for (FSingleParticlePhysicsProxy* UnreadyProxy : UnreadyVolumeParticles_PT)
 	{
@@ -567,7 +574,7 @@ void FSpatialReadinessSimCallback::OnMidPhaseModification_Internal(FMidPhaseModi
 #endif
 
 			// Add the particle to the list of particles to freeze
-			MidPhase_UnreadyRigidParticles_PT.Add(RigidParticle);
+			UnreadyRigidParticles_PT.FindOrAdd(RigidParticle) = 1;
 
 			// Disable the midphase
 			MidPhase.Disable();
@@ -606,25 +613,9 @@ uint32 FSpatialReadinessSimCallback::FHashMapTraits::GetElementID(const FUnready
 
 void FSpatialReadinessSimCallback::ForEachUnreadyRigidParticle_PT(const TFunction<bool(Chaos::FPBDRigidParticleHandle*)>& Lambda) const
 {
-	for (Chaos::FPBDRigidParticleHandle* Particle : NewRegistered_UnreadyRigidParticles_PT)
-	{
-		if (Lambda(Particle) == false)
-		{
-			return;
-		}
-	}
-
-	for (Chaos::FPBDRigidParticleHandle* Particle : OldRegistered_UnreadyRigidParticles_PT)
-	{
-		if (Lambda(Particle) == false)
-		{
-			return;
-		}
-	}
-
-	for (Chaos::FPBDRigidParticleHandle* Particle : MidPhase_UnreadyRigidParticles_PT)
+	for (auto It = UnreadyRigidParticles_PT.CreateConstIterator(); It; ++It)
 	{
-		if (Lambda(Particle) == false)
+		if (Lambda(It->Key) == false)
 		{
 			return;
 		}
@@ -633,9 +624,6 @@ void FSpatialReadinessSimCallback::ForEachUnreadyRigidParticle_PT(const TFunctio
 
 int32 FSpatialReadinessSimCallback::GetNumUnreadyRigidParticles_PT() const
 {
-	return
-		NewRegistered_UnreadyRigidParticles_PT.Num() +
-		OldRegistered_UnreadyRigidParticles_PT.Num() +
-		MidPhase_UnreadyRigidParticles_PT.Num();
+	return UnreadyRigidParticles_PT.Num();
 }
 
diff --git a/Engine/Plugins/Experimental/SpatialReadiness/Source/SpatialReadiness/Public/SpatialReadinessSimCallback.h b/Engine/Plugins/Experimental/SpatialReadiness/Source/SpatialReadiness/Public/SpatialReadinessSimCallback.h
--- a/Engine/Plugins/Experimental/SpatialReadiness/Source/SpatialReadiness/Public/SpatialReadinessSimCallback.h
+++ b/Engine/Plugins/Experimental/SpatialReadiness/Source/SpatialReadiness/Public/SpatialReadinessSimCallback.h
@@ -108,6 +108,10 @@ protected:
 	void FreezeParticles_PT();
 	void UnFreezeParticles_PT();
 
+	// Decrement counters for each unready particle, removing them from the unready
+	// particles array when their counters reach zero.
+	void DecrementUnreadyRigidParticles_PT();
+
 	// Physics thread function for querying for unready volumes
 	// TODO: Make const
 	bool QueryReadiness_PT(const Chaos::FAABB3& Bounds, TArray<const Chaos::FSingleParticlePhysicsProxy*>& OutVolumeProxies);
@@ -128,27 +132,24 @@ protected:
 	// List of particle handles which represent unready volumes
 	TSet<Chaos::FSingleParticlePhysicsProxy*> UnreadyVolumeParticles_PT;
 
-	// Lists of particle handles which represent particles that interacted with unready volumes
-	//
-	// This got a little bit confusing due to some quirks of the particle
-	// creation pipeline.
+	// Lists of particle handles which represent particles that interacted with unready
+	// volumes, along with a counter representing how many more ticks the particle should
+	// remain frozen.
 	//
-	// A particle is added to the NewRegistered list if it's initially
-	// overlapping an unready volume. Ideally it would then get picked
-	// up in the unready list maintained by the mid-phase modifier.
-	// However, it isn't inserted into the acceleration structure
-	// immediately, so it is missed.
+	// The counter strategy is to accommodate the following oddities about particle
+	// creation, deletion, and midphase creation.
 	//
-	// Instead, at the beginning of the next frame, NewRegistered contents
-	// are moved into OldRegistered, and NewRegistered is emptied.
+	//  - A particle can exist for one entire sim tick after it is created without
+	//    appearing in the mid phase. Therefore particles which are registered on top of
+	//    unready volumes must be unreadied with a count of 2.
 	//
-	// This ensures that particles which are spawned on top of unready
-	// volumes will not be able to sneak in a integration step before
-	// they have been inserted into the acceleration structure, where the
-	// mid-phase modifier will be able to manage them.
-	TSet<Chaos::FPBDRigidParticleHandle*> NewRegistered_UnreadyRigidParticles_PT;
-	TSet<Chaos::FPBDRigidParticleHandle*> OldRegistered_UnreadyRigidParticles_PT;
-	TSet<Chaos::FPBDRigidParticleHandle*> MidPhase_UnreadyRigidParticles_PT;
+	//  - For the same reason as the item above, when an unready volume is deleted
+	//    on the same frame that another particle is created (perhaps a streaming object
+	//    with which the unready volume was meant to prevent overlap), there's an
+	//    additional frame during which the new particle has no collision. We lose
+	//    nothing by keeping objects frozen for one more frame. The counter is used
+	//    to accommadate this case.
+	TMap<Chaos::FPBDRigidParticleHandle*, int> UnreadyRigidParticles_PT;
 
 	// "Unready" particles are forced to be stationary in PreSimulate, and restored
 	// to their previous state in PostIntegrate. The values needed for restoration
diff --git a/Engine/Plugins/Experimental/SpatialReadiness/Source/SpatialReadiness/Public/SpatialReadinessStats.h b/Engine/Plugins/Experimental/SpatialReadiness/Source/SpatialReadiness/Public/SpatialReadinessStats.h
--- a/Engine/Plugins/Experimental/SpatialReadiness/Source/SpatialReadiness/Public/SpatialReadinessStats.h
+++ b/Engine/Plugins/Experimental/SpatialReadiness/Source/SpatialReadiness/Public/SpatialReadinessStats.h
@@ -11,6 +11,7 @@ DECLARE_CYCLE_STAT(TEXT("SpatialReadiness_Physics_AddUnreadyVolumeGT"), STAT_Spa
 DECLARE_CYCLE_STAT(TEXT("SpatialReadiness_Physics_RemoveUnreadyVolumeGT"), STAT_SpatialReadiness_Physics_RemoveUnreadyVolumeGT, STATGROUP_SpatialReadiness);
 DECLARE_CYCLE_STAT(TEXT("SpatialReadiness_Physics_FreezeParticlesPT"), STAT_SpatialReadiness_Physics_FreezeParticlesPT, STATGROUP_SpatialReadiness);
 DECLARE_CYCLE_STAT(TEXT("SpatialReadiness_Physics_UnFreezeParticlesPT"), STAT_SpatialReadiness_Physics_UnFreezeParticlesPT, STATGROUP_SpatialReadiness);
+DECLARE_CYCLE_STAT(TEXT("SpatialReadiness_Physics_DecrementUnreadyRigidParticlesPT"), STAT_SpatialReadiness_Physics_DecrementUnreadyRigidParticlesPT, STATGROUP_SpatialReadiness);
 DECLARE_CYCLE_STAT(TEXT("SpatialReadiness_Physics_PreSimulate"), STAT_SpatialReadiness_Physics_PreSimulate, STATGROUP_SpatialReadiness);
 DECLARE_CYCLE_STAT(TEXT("SpatialReadiness_Physics_ParticlesRegistered"), STAT_SpatialReadiness_Physics_ParticlesRegistered, STATGROUP_SpatialReadiness);
 DECLARE_CYCLE_STAT(TEXT("SpatialReadiness_Physics_MidPhase"), STAT_SpatialReadiness_Physics_MidPhase, STATGROUP_SpatialReadiness);
