Sim tick freeze frame cvar, default to 2. Add check for SyncTimeStamp validity as well.

diff --git a/Engine/Plugins/Experimental/SpatialReadiness/Source/SpatialReadiness/Private/SpatialReadiness/SpatialReadinessSimCallback.cpp b/Engine/Plugins/Experimental/SpatialReadiness/Source/SpatialReadiness/Private/SpatialReadiness/SpatialReadinessSimCallback.cpp
--- a/Engine/Plugins/Experimental/SpatialReadiness/Source/SpatialReadiness/Private/SpatialReadiness/SpatialReadinessSimCallback.cpp
+++ b/Engine/Plugins/Experimental/SpatialReadiness/Source/SpatialReadiness/Private/SpatialReadiness/SpatialReadinessSimCallback.cpp
@@ -26,6 +26,45 @@
 
 using namespace Chaos;
 
+namespace
+{
+	bool IsParticleValid(Chaos::FGeometryParticleHandle* ParticleHandle, TSharedPtr<FProxyTimestampBase, ESPMode::ThreadSafe> SyncTimestamp)
+	{
+		if (ParticleHandle == nullptr)
+		{
+			return false;
+		}
+
+		if (SyncTimestamp.IsValid() == false)
+		{
+			return false;
+		}
+
+		if (SyncTimestamp->bDeleted)
+		{
+			return false;
+		}
+
+		IPhysicsProxyBase* Proxy = ParticleHandle->PhysicsProxy();
+		if (Proxy == nullptr)
+		{
+			return false;
+		}
+
+		if (!ensureMsgf(Proxy->GetSyncTimestamp() == SyncTimestamp, TEXT("particle's sync timestamp doesn't match the sync timestamp that was passed to the particle validity check!")))
+		{
+			return false;
+		}
+
+		return true;
+	};
+}
+
+TAutoConsoleVariable<int32> CVarSpatialReadiness_SimTickFreezeFrames(
+	TEXT("p.SpatialReadiness.SimTickFreezeFrames"), 2,
+	TEXT("When a rigid particle is frozen by contact with an unready volume, this is the number of frames which it will remain frozen after losing contact. Theoretically it could be zero, but chaos can take 1-2 frames to fully initialize the collision of new particles, and so we default it to 2 to be safe. Please avoid the temptation to use this as an indefinitely extendable \"freeze timer\" to cover up more serious timing issues in systems which use readiness."),
+	ECVF_Default);
+
 FSpatialReadinessSimCallback::FSpatialReadinessSimCallback(FPhysScene_Chaos& InPhysicsScene)
 	: PhysicsScene(InPhysicsScene)
 	, UnreadyVolumeData_GT(256)
@@ -369,7 +408,7 @@ void FSpatialReadinessSimCallback::FreezeParticles_PT()
 		return;
 	}
 
-	ForEachUnreadyRigidParticle_PT([&](Chaos::FPBDRigidParticleHandle* RigidParticle)
+	ForEachUnreadyRigidParticle_PT([&](Chaos::FPBDRigidParticleHandle* RigidParticle, const int FreezeFrames)
 	{
 		const EObjectStateType ObjectState = RigidParticle->ObjectState();
 		if (ObjectState == EObjectStateType::Static)
@@ -401,7 +440,7 @@ void FSpatialReadinessSimCallback::FreezeParticles_PT()
 		{
 			// Draw a box in CVD representing the frozen object
 			SCOPE_CYCLE_COUNTER(STAT_SpatialReadiness_Physics_ChaosVisualDebugger)
-			const FName Tag = *RigidParticle->GetDebugName();
+			const FName Tag = *FString::Printf(TEXT("%s (%d)"), *RigidParticle->GetDebugName().Left(25), FreezeFrames);
 			const FAABB3 Bounds = RigidParticle->WorldSpaceInflatedBounds();
 			CVD_TRACE_DEBUG_DRAW_BOX(FBox(Bounds.Min(), Bounds.Max()), Tag, FColor::Yellow, FChaosVisualDebuggerTrace::GetSolverID(*GetSolver()));
 		}
@@ -446,7 +485,7 @@ void FSpatialReadinessSimCallback::DecrementUnreadyRigidParticles_PT()
 	// Decrement each counter, 
 	for (auto It = UnreadyRigidParticles_PT.CreateIterator(); It; ++It)
 	{
-		if (--It->Value < 0)
+		if (IsParticleValid(It->Key, It->Value.SyncTimeStamp.Pin()) == false || --It->Value.FreezeFrames < 0)
 		{
 			It.RemoveCurrent();
 		}
@@ -480,6 +519,9 @@ void FSpatialReadinessSimCallback::OnParticlesRegistered_Internal(TArray<FSingle
 {
 	SCOPE_CYCLE_COUNTER(STAT_SpatialReadiness_Physics_ParticlesRegistered)
 
+	// The number of frames we keep frozen particles frozen
+	const int32 FreezeFrames = FMath::Max(1, CVarSpatialReadiness_SimTickFreezeFrames.GetValueOnAnyThread());
+
 	// For each newly added particle, query against unready volumes to see if
 	// it should be frozen.
 	for (FSingleParticlePhysicsProxy* ParticleProxy : RegisteredProxies)
@@ -496,14 +538,22 @@ void FSpatialReadinessSimCallback::OnParticlesRegistered_Internal(TArray<FSingle
 			continue;
 		}
 
+		TWeakPtr<FProxyTimestampBase, ESPMode::ThreadSafe> SyncTimeStamp = ParticleProxy->GetSyncTimestamp();
+		if (SyncTimeStamp.IsValid() == false)
+		{
+			continue;
+		}
+
 		// Get the bounds to use for the query
 		TArray<const FSingleParticlePhysicsProxy*> VolumeProxies;
 		const FAABB3 QueryBounds = RigidParticle->WorldSpaceInflatedBounds();
 		if (!QueryReadiness_PT(QueryBounds, VolumeProxies))
 		{
 			// If this volume hit an unready volume, query readiness returns false
-			// and we must mark the particle as frozen
-			UnreadyRigidParticles_PT.FindOrAdd(RigidParticle) = 2;
+			// and we must mark the particle as frozen. We freeze for "+1" frames
+			// because it actually takes 2 frames for new particles to show up in
+			// midphases.
+			UnreadyRigidParticles_PT.FindOrAdd(RigidParticle) = { SyncTimeStamp, FreezeFrames + 1 };
 		}
 	}
 }
@@ -512,6 +562,9 @@ void FSpatialReadinessSimCallback::OnMidPhaseModification_Internal(FMidPhaseModi
 {
 	SCOPE_CYCLE_COUNTER(STAT_SpatialReadiness_Physics_MidPhase)
 
+	// The number of frames to keep particles we freeze frozen
+	const int32 FreezeFrames = FMath::Max(1, CVarSpatialReadiness_SimTickFreezeFrames.GetValueOnAnyThread());
+
 	// Go through every unready volume
 	for (FSingleParticlePhysicsProxy* UnreadyProxy : UnreadyVolumeParticles_PT)
 	{
@@ -564,6 +617,18 @@ void FSpatialReadinessSimCallback::OnMidPhaseModification_Internal(FMidPhaseModi
 				continue;
 			}
 
+			FSingleParticlePhysicsProxy* RigidProxy = static_cast<FSingleParticlePhysicsProxy*>(RigidParticle->PhysicsProxy());
+			if (RigidProxy == nullptr)
+			{
+				continue;
+			}
+
+			TWeakPtr<FProxyTimestampBase, ESPMode::ThreadSafe> SyncTimeStamp = RigidProxy->GetSyncTimestamp();
+			if (SyncTimeStamp.IsValid() == false)
+			{
+				continue;
+			}
+
 #if ENABLE_DRAW_DEBUG
 			if (CVarSpatialReadinessDebugDraw.GetValueOnAnyThread())
 			{
@@ -574,7 +639,7 @@ void FSpatialReadinessSimCallback::OnMidPhaseModification_Internal(FMidPhaseModi
 #endif
 
 			// Add the particle to the list of particles to freeze
-			UnreadyRigidParticles_PT.FindOrAdd(RigidParticle) = 1;
+			UnreadyRigidParticles_PT.FindOrAdd(RigidParticle) = { SyncTimeStamp, FreezeFrames };
 
 			// Disable the midphase
 			MidPhase.Disable();
@@ -612,10 +677,20 @@ uint32 FSpatialReadinessSimCallback::FHashMapTraits::GetElementID(const FUnready
 }
 
 void FSpatialReadinessSimCallback::ForEachUnreadyRigidParticle_PT(const TFunction<bool(Chaos::FPBDRigidParticleHandle*)>& Lambda) const
+{
+	ForEachUnreadyRigidParticle_PT([&Lambda](Chaos::FPBDRigidParticleHandle* Particle, int) { return Lambda(Particle); });
+}
+
+void FSpatialReadinessSimCallback::ForEachUnreadyRigidParticle_PT(const TFunction<bool(Chaos::FPBDRigidParticleHandle*, int)>& Lambda) const
 {
 	for (auto It = UnreadyRigidParticles_PT.CreateConstIterator(); It; ++It)
 	{
-		if (Lambda(It->Key) == false)
+		if (IsParticleValid(It->Key, It->Value.SyncTimeStamp.Pin()) == false)
+		{
+			continue;
+		}
+	
+		if (Lambda(It->Key, It->Value.FreezeFrames) == false)
 		{
 			return;
 		}
diff --git a/Engine/Plugins/Experimental/SpatialReadiness/Source/SpatialReadiness/Public/SpatialReadinessSimCallback.h b/Engine/Plugins/Experimental/SpatialReadiness/Source/SpatialReadiness/Public/SpatialReadinessSimCallback.h
--- a/Engine/Plugins/Experimental/SpatialReadiness/Source/SpatialReadiness/Public/SpatialReadinessSimCallback.h
+++ b/Engine/Plugins/Experimental/SpatialReadiness/Source/SpatialReadiness/Public/SpatialReadinessSimCallback.h
@@ -86,6 +86,9 @@ public:
 	// Perform an operation on each of the currently frozen particles - only valid on the physics thread
 	void ForEachUnreadyRigidParticle_PT(const TFunction<bool(Chaos::FPBDRigidParticleHandle*)>& Lambda) const;
 
+	// A variant which also sees how many frames each particle will remain frozen
+	void ForEachUnreadyRigidParticle_PT(const TFunction<bool(Chaos::FPBDRigidParticleHandle*, int32)>& Lambda) const;
+
 	// Get the total number of currently frozen particles - only valid on the physics thread
 	int32 GetNumUnreadyRigidParticles_PT() const;
 
@@ -149,7 +152,17 @@ protected:
 	//    additional frame during which the new particle has no collision. We lose
 	//    nothing by keeping objects frozen for one more frame. The counter is used
 	//    to accommadate this case.
-	TMap<Chaos::FPBDRigidParticleHandle*, int> UnreadyRigidParticles_PT;
+	//
+	// This freeze frame counter is packed into a small struct which also contains the
+	// SyncTimeStamp for the particle's proxy, because if we're going to track the particle
+	// for more than one frame (and we are), then we also need to track the validity of the
+	// SyncTimeStamp before ever accessing methods or properties on the particle handle ptr.
+	struct FUnreadyRigidParticleData
+	{
+		TWeakPtr<FProxyTimestampBase, ESPMode::ThreadSafe> SyncTimeStamp;
+		int FreezeFrames;
+	};
+	TMap<Chaos::FPBDRigidParticleHandle*, FUnreadyRigidParticleData> UnreadyRigidParticles_PT;
 
 	// "Unready" particles are forced to be stationary in PreSimulate, and restored
 	// to their previous state in PostIntegrate. The values needed for restoration
