Changes to make perf tests for SR more meaningful:
- Move SR macros into a separate header, add WITH_SPATIAL_READINESS_CVD macro which enables cvd only in non-test, non-shipping builds.
- Make innermost cycle stat counters "verbose" so that they don't affect perf.

diff --git a/Engine/Plugins/Experimental/SpatialReadiness/Source/SpatialReadiness/Private/SpatialReadiness/SpatialReadinessDebug.cpp b/Engine/Plugins/Experimental/SpatialReadiness/Source/SpatialReadiness/Private/SpatialReadiness/SpatialReadinessDebug.cpp
--- a/Engine/Plugins/Experimental/SpatialReadiness/Source/SpatialReadiness/Private/SpatialReadiness/SpatialReadinessDebug.cpp
+++ b/Engine/Plugins/Experimental/SpatialReadiness/Source/SpatialReadiness/Private/SpatialReadiness/SpatialReadinessDebug.cpp
@@ -86,7 +86,11 @@ static FAutoConsoleCommandWithWorldAndArgs CmdSpatialReadinessPrint(
 
 				const FAABB3 Bounds = UnreadyVolume->WorldSpaceInflatedBounds();
 
-				Builder << " " << (VolumeNum++) << ". " << Data.Description << " : " << *Bounds.ToString() << "\n";
+				Builder << " " << (VolumeNum++) << ". "
+#if WITH_SPATIAL_READINESS_DESCRIPTIONS
+				<< Data.Description << " : "
+#endif
+				<< *Bounds.ToString() << "\n";
 			});
 
 			// Print the string
diff --git a/Engine/Plugins/Experimental/SpatialReadiness/Source/SpatialReadiness/Private/SpatialReadiness/SpatialReadinessSimCallback.cpp b/Engine/Plugins/Experimental/SpatialReadiness/Source/SpatialReadiness/Private/SpatialReadiness/SpatialReadinessSimCallback.cpp
--- a/Engine/Plugins/Experimental/SpatialReadiness/Source/SpatialReadiness/Private/SpatialReadiness/SpatialReadinessSimCallback.cpp
+++ b/Engine/Plugins/Experimental/SpatialReadiness/Source/SpatialReadiness/Private/SpatialReadiness/SpatialReadinessSimCallback.cpp
@@ -23,6 +23,7 @@
 #include "Chaos/DebugDrawQueue.h"
 #include "SpatialReadinessLog.h"
 #include "SpatialReadinessDebug.h"
+#include "SpatialReadinessDefines.h"
 
 using namespace Chaos;
 
@@ -160,9 +161,9 @@ int32 FSpatialReadinessSimCallback::AddUnreadyVolume_GT(const FBox& Bounds, cons
 	SimInput.UnreadyVolumesToRemove.Remove(ParticleProxy);
 
 	// Log the event
-	UE_LOG(LogSpatialReadiness, Log, TEXT("Unready volume created.with index [%d]: %s"), VolumeIndex, *Description);
+	UE_LOG(LogSpatialReadiness, Verbose, TEXT("Unready volume created.with index [%d]: %s"), VolumeIndex, *Description);
 
-#if !UE_BUILD_SHIPPING
+#if WITH_SPATIAL_READINESS_CVD
 	{
 		// Draw a box in CVD representing the unready volume
 		SCOPE_CYCLE_COUNTER(STAT_SpatialReadiness_Physics_ChaosVisualDebugger)
@@ -185,7 +186,7 @@ void FSpatialReadinessSimCallback::RemoveUnreadyVolume_GT(int32 UnreadyVolumeInd
 		return;
 	}
 
-#if !UE_BUILD_SHIPPING
+#if WITH_SPATIAL_READINESS_CVD
 	{
 		SCOPE_CYCLE_COUNTER(STAT_SpatialReadiness_Physics_ChaosVisualDebugger)
 #if WITH_SPATIAL_READINESS_DESCRIPTIONS
@@ -227,7 +228,7 @@ void FSpatialReadinessSimCallback::RemoveUnreadyVolume_GT(int32 UnreadyVolumeInd
 	FChaosEngineInterface::ReleaseActor(ParticleProxy, &PhysicsScene);
 
 	// Log the event
-	UE_LOG(LogSpatialReadiness, Log, TEXT("Unready volume removed with index [%d]"), UnreadyVolumeIndex);
+	UE_LOG(LogSpatialReadiness, Verbose, TEXT("Unready volume removed with index [%d]"), UnreadyVolumeIndex);
 }
 
 bool FSpatialReadinessSimCallback::QueryReadiness_GT(const FBox& Bounds, TArray<int32>& OutVolumeIndices, bool bAllUnreadyVolumes) const
@@ -436,7 +437,7 @@ void FSpatialReadinessSimCallback::FreezeParticles_PT()
 		}
 #endif
 
-#if !UE_BUILD_SHIPPING
+#if WITH_SPATIAL_READINESS_CVD
 		{
 			// Draw a box in CVD representing the frozen object
 			SCOPE_CYCLE_COUNTER(STAT_SpatialReadiness_Physics_ChaosVisualDebugger)
@@ -588,7 +589,7 @@ void FSpatialReadinessSimCallback::OnMidPhaseModification_Internal(FMidPhaseModi
 		}
 #endif
 
-#if !UE_BUILD_SHIPPING
+#if WITH_SPATIAL_READINESS_CVD
 		{
 			SCOPE_CYCLE_COUNTER(STAT_SpatialReadiness_Physics_ChaosVisualDebugger)
 			// Draw a box in CVD representing the unready volume
@@ -599,50 +600,52 @@ void FSpatialReadinessSimCallback::OnMidPhaseModification_Internal(FMidPhaseModi
 		}
 #endif
 
-		// Go through every mid-phase which involves this volume
-		for (FMidPhaseModifier& MidPhase : Accessor.GetMidPhases(UnreadyVolume))
 		{
 			SCOPE_CYCLE_COUNTER(STAT_SpatialReadiness_Physics_MidPhase_AccessorLoop)
 
-			// Get the particle that is not the unready volume
-			FGeometryParticleHandle* GeometryParticle = MidPhase.GetOtherParticle(UnreadyVolume);
-			if (GeometryParticle == nullptr)
+			// Go through every mid-phase which involves this volume
+			for (FMidPhaseModifier& MidPhase : Accessor.GetMidPhases(UnreadyVolume))
 			{
-				continue;
-			}
-
-			FPBDRigidParticleHandle* RigidParticle = GeometryParticle->CastToRigidParticle();
-			if (RigidParticle == nullptr)
-			{
-				continue;
-			}
-
-			FSingleParticlePhysicsProxy* RigidProxy = static_cast<FSingleParticlePhysicsProxy*>(RigidParticle->PhysicsProxy());
-			if (RigidProxy == nullptr)
-			{
-				continue;
-			}
-
-			TWeakPtr<FProxyTimestampBase, ESPMode::ThreadSafe> SyncTimeStamp = RigidProxy->GetSyncTimestamp();
-			if (SyncTimeStamp.IsValid() == false)
-			{
-				continue;
-			}
+				// Get the particle that is not the unready volume
+				FGeometryParticleHandle* GeometryParticle = MidPhase.GetOtherParticle(UnreadyVolume);
+				if (GeometryParticle == nullptr)
+				{
+					continue;
+				}
+
+				FPBDRigidParticleHandle* RigidParticle = GeometryParticle->CastToRigidParticle();
+				if (RigidParticle == nullptr)
+				{
+					continue;
+				}
+
+				FSingleParticlePhysicsProxy* RigidProxy = static_cast<FSingleParticlePhysicsProxy*>(RigidParticle->PhysicsProxy());
+				if (RigidProxy == nullptr)
+				{
+					continue;
+				}
+
+				TWeakPtr<FProxyTimestampBase, ESPMode::ThreadSafe> SyncTimeStamp = RigidProxy->GetSyncTimestamp();
+				if (SyncTimeStamp.IsValid() == false)
+				{
+					continue;
+				}
 
 #if ENABLE_DRAW_DEBUG
-			if (CVarSpatialReadinessDebugDraw.GetValueOnAnyThread())
-			{
-				SCOPE_CYCLE_COUNTER(STAT_SpatialReadiness_Physics_DebugDraw)
-				const FAABB3 Bounds = RigidParticle->WorldSpaceInflatedBounds();
-				Chaos::FDebugDrawQueue::GetInstance().DrawDebugBox(Bounds.Center(), Bounds.Extents() * .5f, FQuat::Identity, FColor::Yellow, false, -1.f, -1, 0.f);
-			}
+				if (CVarSpatialReadinessDebugDraw.GetValueOnAnyThread())
+				{
+					SCOPE_CYCLE_COUNTER(STAT_SpatialReadiness_Physics_DebugDraw)
+					const FAABB3 Bounds = RigidParticle->WorldSpaceInflatedBounds();
+					Chaos::FDebugDrawQueue::GetInstance().DrawDebugBox(Bounds.Center(), Bounds.Extents() * .5f, FQuat::Identity, FColor::Yellow, false, -1.f, -1, 0.f);
+				}
 #endif
 
-			// Add the particle to the list of particles to freeze
-			UnreadyRigidParticles_PT.FindOrAdd(RigidParticle) = { SyncTimeStamp, FreezeFrames };
+				// Add the particle to the list of particles to freeze
+				UnreadyRigidParticles_PT.FindOrAdd(RigidParticle) = { SyncTimeStamp, FreezeFrames };
 
-			// Disable the midphase
-			MidPhase.Disable();
+				// Disable the midphase
+				MidPhase.Disable();
+			}
 		}
 	}
 }
diff --git a/Engine/Plugins/Experimental/SpatialReadiness/Source/SpatialReadiness/Public/SpatialReadinessDefines.h b/Engine/Plugins/Experimental/SpatialReadiness/Source/SpatialReadiness/Public/SpatialReadinessDefines.h
new file mode 100644
--- /dev/null
+++ b/Engine/Plugins/Experimental/SpatialReadiness/Source/SpatialReadiness/Public/SpatialReadinessDefines.h
@@ -0,0 +1,13 @@
+// Copyright Epic Games, Inc. All Rights Reserved.
+#pragma once
+
+#include "CoreMinimal.h"
+
+#ifndef WITH_SPATIAL_READINESS_DESCRIPTIONS
+#define WITH_SPATIAL_READINESS_DESCRIPTIONS (!(UE_BUILD_SHIPPING || UE_BUILD_TEST))
+#endif
+
+#ifndef WITH_SPATIAL_READINESS_CVD
+#define WITH_SPATIAL_READINESS_CVD (!(UE_BUILD_SHIPPING || UE_BUILD_TEST))
+#endif
+
diff --git a/Engine/Plugins/Experimental/SpatialReadiness/Source/SpatialReadiness/Public/SpatialReadinessStats.h b/Engine/Plugins/Experimental/SpatialReadiness/Source/SpatialReadiness/Public/SpatialReadinessStats.h
--- a/Engine/Plugins/Experimental/SpatialReadiness/Source/SpatialReadiness/Public/SpatialReadinessStats.h
+++ b/Engine/Plugins/Experimental/SpatialReadiness/Source/SpatialReadiness/Public/SpatialReadinessStats.h
@@ -5,6 +5,7 @@
 #include "Stats/Stats.h"
 
 DECLARE_STATS_GROUP(TEXT("SpatialReadiness"), STATGROUP_SpatialReadiness, STATCAT_Advanced);
+DECLARE_STATS_GROUP_VERBOSE(TEXT("SpatialReadinessDetailed"), STATGROUP_SpatialReadinessVerbose, STATCAT_Advanced)
 
 // Physics
 DECLARE_CYCLE_STAT(TEXT("SpatialReadiness_Physics_AddUnreadyVolumeGT"), STAT_SpatialReadiness_Physics_AddUnreadyVolumeGT, STATGROUP_SpatialReadiness);
@@ -15,7 +16,7 @@ DECLARE_CYCLE_STAT(TEXT("SpatialReadiness_Physics_DecrementUnreadyRigidParticles
 DECLARE_CYCLE_STAT(TEXT("SpatialReadiness_Physics_PreSimulate"), STAT_SpatialReadiness_Physics_PreSimulate, STATGROUP_SpatialReadiness);
 DECLARE_CYCLE_STAT(TEXT("SpatialReadiness_Physics_ParticlesRegistered"), STAT_SpatialReadiness_Physics_ParticlesRegistered, STATGROUP_SpatialReadiness);
 DECLARE_CYCLE_STAT(TEXT("SpatialReadiness_Physics_MidPhase"), STAT_SpatialReadiness_Physics_MidPhase, STATGROUP_SpatialReadiness);
-DECLARE_CYCLE_STAT(TEXT("SpatialReadiness_Physics_MidPhase_AccessorLoop"), STAT_SpatialReadiness_Physics_MidPhase_AccessorLoop, STATGROUP_SpatialReadiness);
+DECLARE_CYCLE_STAT(TEXT("SpatialReadiness_Physics_MidPhase_AccessorLoop"), STAT_SpatialReadiness_Physics_MidPhase_AccessorLoop, STATGROUP_SpatialReadinessVerbose);
 DECLARE_CYCLE_STAT(TEXT("SpatialReadiness_Physics_ChaosVisualDebugger"), STAT_SpatialReadiness_Physics_ChaosVisualDebugger, STATGROUP_SpatialReadiness);
 DECLARE_CYCLE_STAT(TEXT("SpatialReadiness_Physics_DebugDraw"), STAT_SpatialReadiness_Physics_DebugDraw, STATGROUP_SpatialReadiness);
 
diff --git a/Engine/Plugins/Experimental/SpatialReadiness/Source/SpatialReadiness/Public/SpatialReadinessVolume.h b/Engine/Plugins/Experimental/SpatialReadiness/Source/SpatialReadiness/Public/SpatialReadinessVolume.h
--- a/Engine/Plugins/Experimental/SpatialReadiness/Source/SpatialReadiness/Public/SpatialReadinessVolume.h
+++ b/Engine/Plugins/Experimental/SpatialReadiness/Source/SpatialReadiness/Public/SpatialReadinessVolume.h
@@ -4,10 +4,7 @@
 
 #include "CoreMinimal.h"
 #include "Math/MathFwd.h"
-
-#ifndef WITH_SPATIAL_READINESS_DESCRIPTIONS
-#define WITH_SPATIAL_READINESS_DESCRIPTIONS (!UE_BUILD_SHIPPING)
-#endif
+#include "SpatialReadinessDefines.h"
 
 struct FSpatialReadinessAPIDelegates;
 
