turn on spatial readinesss on dedicated servers, and turn off world readiness
register SR show flag even if not in editor, only register draw callback when WITH_EDITOR

diff --git a/Engine/Plugins/Experimental/SpatialReadiness/Source/SpatialReadiness/Private/SpatialReadiness/SpatialReadinessSubsystem.cpp b/Engine/Plugins/Experimental/SpatialReadiness/Source/SpatialReadiness/Private/SpatialReadiness/SpatialReadinessSubsystem.cpp
--- a/Engine/Plugins/Experimental/SpatialReadiness/Source/SpatialReadiness/Private/SpatialReadiness/SpatialReadinessSubsystem.cpp
+++ b/Engine/Plugins/Experimental/SpatialReadiness/Source/SpatialReadiness/Private/SpatialReadiness/SpatialReadinessSubsystem.cpp
@@ -15,11 +15,14 @@
 
 #include UE_INLINE_GENERATED_CPP_BY_NAME(SpatialReadinessSubsystem)
 
-TAutoConsoleVariable<bool> CVarSpatialReadiness_OnDedicatedServer(TEXT("p.SpatialReadiness.OnDedicatedServer"), 0, TEXT("Enable spatial readiness on dedicated servers"),  ECVF_Default);
+TAutoConsoleVariable<bool> CVarSpatialReadiness_OnDedicatedServer(TEXT("p.SpatialReadiness.OnDedicatedServer"), 1, TEXT("Enable spatial readiness on dedicated servers"),  ECVF_Default);
 TAutoConsoleVariable<bool> CVarSpatialReadiness_OnStandalone(     TEXT("p.SpatialReadiness.OnStandalone"),      0, TEXT("Enable spatial readiness on standalone servers"), ECVF_Default);
 TAutoConsoleVariable<bool> CVarSpatialReadiness_OnClient(         TEXT("p.SpatialReadiness.OnClient"),          0, TEXT("Enable spatial readiness on clients"),            ECVF_Default);
 TAutoConsoleVariable<bool> CVarSpatialReadiness_Debug_NeverReady( TEXT("p.SpatialReadiness.Debug.NeverReady"),  0, TEXT("When true, unready volumes are never removed"),   ECVF_Default);
 
+const FString SpatialReadinessShowFlagName = TEXT("SpatialReadiness");
+TCustomShowFlag<> SpatialReadinessShowFlag(*SpatialReadinessShowFlagName, false, SFG_Developer, NSLOCTEXT("UnrealEd", "SpatialReadinessShowFlag", "Spatial Readiness"));
+
 using namespace Chaos;
 
 USpatialReadiness::USpatialReadiness()
@@ -65,9 +68,10 @@ void USpatialReadiness::OnWorldBeginPlay(UWorld& InWorld)
 	CreateSimCallback();
 }
 
-#if ENABLE_DRAW_DEBUG
+#if WITH_EDITOR
 void USpatialReadiness::OnDebugDraw(UCanvas* Canvas, APlayerController* PlayerController)
 {
+#if ENABLE_DRAW_DEBUG
 	if (Canvas == nullptr) { return; }
 	if (PlayerController == nullptr) {
 		PlayerController = GetWorld()->GetFirstPlayerController();
@@ -128,23 +132,22 @@ void USpatialReadiness::OnDebugDraw(UCanvas* Canvas, APlayerController* PlayerCo
 			Y += FontHeight + 4;
 		}
 	}
+#endif // ENABLE_DRAW_DEBUG
 }
-#endif
+#endif // WITH_EDITOR
 
 void USpatialReadiness::Initialize(FSubsystemCollectionBase& Collection)
 {
 	Super::Initialize(Collection);
 
-#if ENABLE_DRAW_DEBUG
-	const FString Tag = TEXT("SpatialReadiness");
-	FEngineShowFlags::RegisterCustomShowFlag(*Tag, false, EShowFlagGroup::SFG_Developer, NSLOCTEXT("UnrealEd", "SpatialReadinessSF", "Spatial Readiness"));
-	OnDebugDrawDelegateHandle = UDebugDrawService::Register(*Tag, FDebugDrawDelegate::CreateUObject(this, &USpatialReadiness::OnDebugDraw));
+#if WITH_EDITOR
+	OnDebugDrawDelegateHandle = UDebugDrawService::Register(*SpatialReadinessShowFlagName, FDebugDrawDelegate::CreateUObject(this, &USpatialReadiness::OnDebugDraw));
 #endif
 }
 
 void USpatialReadiness::Deinitialize()
 {
-#if ENABLE_DRAW_DEBUG
+#if WITH_EDITOR
 	UDebugDrawService::Unregister(OnDebugDrawDelegateHandle);
 #endif
 	DestroySimCallback();
diff --git a/Engine/Plugins/Experimental/SpatialReadiness/Source/SpatialReadiness/Public/SpatialReadinessSubsystem.h b/Engine/Plugins/Experimental/SpatialReadiness/Source/SpatialReadiness/Public/SpatialReadinessSubsystem.h
--- a/Engine/Plugins/Experimental/SpatialReadiness/Source/SpatialReadiness/Public/SpatialReadinessSubsystem.h
+++ b/Engine/Plugins/Experimental/SpatialReadiness/Source/SpatialReadiness/Public/SpatialReadinessSubsystem.h
@@ -47,7 +47,7 @@ public:
 	const FSpatialReadinessSimCallback* GetSimCallback() const { return SimCallback; }
 #endif
 
-#if ENABLE_DRAW_DEBUG
+#if WITH_EDITOR
 	FDelegateHandle OnDebugDrawDelegateHandle;
 	void OnDebugDraw(UCanvas* Canvas, APlayerController* PlayerController);
 #endif
