Invert spatial readiness midphase modifier embedded loops. Add cvar for toggling between midphase strategies. Add function for running heuristic swatch, once we have a good heuristic. Add empty function for grid once we have grid optimization ready to go.

diff --git a/Engine/Plugins/Experimental/SpatialReadiness/Source/SpatialReadiness/Private/SpatialReadiness/SpatialReadinessSimCallback.cpp b/Engine/Plugins/Experimental/SpatialReadiness/Source/SpatialReadiness/Private/SpatialReadiness/SpatialReadinessSimCallback.cpp
--- a/Engine/Plugins/Experimental/SpatialReadiness/Source/SpatialReadiness/Private/SpatialReadiness/SpatialReadinessSimCallback.cpp
+++ b/Engine/Plugins/Experimental/SpatialReadiness/Source/SpatialReadiness/Private/SpatialReadiness/SpatialReadinessSimCallback.cpp
@@ -66,6 +66,16 @@ TAutoConsoleVariable<int32> CVarSpatialReadiness_SimTickFreezeFrames(
 	TEXT("When a rigid particle is frozen by contact with an unready volume, this is the number of frames which it will remain frozen after losing contact. Theoretically it could be zero, but chaos can take 1-2 frames to fully initialize the collision of new particles, and so we default it to 2 to be safe. Please avoid the temptation to use this as an indefinitely extendable \"freeze timer\" to cover up more serious timing issues in systems which use readiness."),
 	ECVF_Default);
 
+TAutoConsoleVariable<int32> CVarSpatialReadiness_MidPhaseStrategy(
+	TEXT("p.SpatialReadiness.MidPhaseStrategy"), 3,
+	TEXT("0: Loop over unready volumes first, midphases second. 1: Loop over all midphases, unready volumes second. 2: (NOT IMPLEMENTED YET) Use a grid acceleration structure for a more optimal version of the unready-volumes-first approach. 3: Dynamically choose between strategies based on how many midphases there are and vs unready volumes."),
+	ECVF_Default);
+
+TAutoConsoleVariable<float> CVarSpatialReadiness_MidPhaseStrategy_Dynamic_ThresholdRatioMidPhasesToUnreadyVolumes(
+	TEXT("p.SpatialReadiness.MidPhaseStrategy.Dynamic.ThresholdRatioMidPhasesToUnreadyVolumes"), 1,
+	TEXT("The ratio of mid phases to unready volumes below which we prioritize mid phases and above which we prioritize unready volumes."),
+	ECVF_Default);
+
 FSpatialReadinessSimCallback::FSpatialReadinessSimCallback(FPhysScene_Chaos& InPhysicsScene)
 	: PhysicsScene(InPhysicsScene)
 	, UnreadyVolumeData_GT(256)
@@ -299,6 +309,24 @@ bool FSpatialReadinessSimCallback::QueryReadiness_GT(const FBox& Bounds, TArray<
 	return OutVolumeIndices.Num() == 0;
 }
 
+bool FSpatialReadinessSimCallback::IsUnreadyVolume_PT(const FGeometryParticleHandle& ParticleHandle) const
+{
+	// Start with the cheapest checks... eliminate particles which cannot be unready volumes
+	// based on state
+	if (ParticleHandle.ObjectState() != EObjectStateType::Static)
+	{
+		return false;
+	}
+
+	// The most expensive check - see if the particle is in our set of unready volumes
+	if (!UnreadyVolumeParticles_PT.Contains(static_cast<const FSingleParticlePhysicsProxy*>(ParticleHandle.PhysicsProxy())))
+	{
+		return false;
+	}
+
+	return true;
+}
+
 const FUnreadyVolumeData_GT* FSpatialReadinessSimCallback::GetVolumeData_GT(int32 VolumeIndex) const
 {
 	return UnreadyVolumeData_GT.Find(VolumeIndex);
@@ -562,8 +590,58 @@ void FSpatialReadinessSimCallback::OnMidPhaseModification_Internal(FMidPhaseModi
 {
 	SCOPE_CYCLE_COUNTER(STAT_SpatialReadiness_Physics_MidPhase)
 
-	// The number of frames to keep particles we freeze frozen
-	const int32 FreezeFrames = FMath::Max(1, CVarSpatialReadiness_SimTickFreezeFrames.GetValueOnAnyThread());
+	switch (GetMidPhaseStrategy(Accessor))
+	{
+		case ESpatialReadinessSimCallback_MidPhaseStrategy::PrioritizeUnreadyVolumes:
+			OnMidPhaseModification_PrioritizeUnreadyVolumes(Accessor);
+			break;
+
+		case ESpatialReadinessSimCallback_MidPhaseStrategy::PrioritizeMidPhases:
+			OnMidPhaseModification_PrioritizeMidPhases(Accessor);
+			break;
+
+		case ESpatialReadinessSimCallback_MidPhaseStrategy::Grid:
+			OnMidPhaseModification_Grid(Accessor);
+			break;
+	}
+}
+
+int32 FSpatialReadinessSimCallback::GetFreezeFrames()
+{
+	return FMath::Max(1, CVarSpatialReadiness_SimTickFreezeFrames.GetValueOnAnyThread());
+}
+
+ESpatialReadinessSimCallback_MidPhaseStrategy FSpatialReadinessSimCallback::GetMidPhaseStrategy(const Chaos::FMidPhaseModifierAccessor& Accessor) const
+{
+	// TODO: For now we just manually pick a strategy using a CVar.
+	//       We should instead use some heuristic in here to select the best strategy
+	//       for the given world setup.
+	const int32 MidPhaseStrategy = CVarSpatialReadiness_MidPhaseStrategy.GetValueOnAnyThread();
+	if (MidPhaseStrategy < static_cast<int32>(ESpatialReadinessSimCallback_MidPhaseStrategy::COUNT))
+	{
+		return ESpatialReadinessSimCallback_MidPhaseStrategy(MidPhaseStrategy);
+	}
+	else if (MidPhaseStrategy == 3)
+	{
+		const int32 NumMidPhases = Accessor.GetNumMidPhases();
+		const int32 NumUnreadyVolumes = UnreadyVolumeParticles_PT.Num();
+		const float RatioThreshold = CVarSpatialReadiness_MidPhaseStrategy_Dynamic_ThresholdRatioMidPhasesToUnreadyVolumes.GetValueOnAnyThread();
+		const float Ratio
+			= NumUnreadyVolumes > 0
+			? static_cast<float>(NumMidPhases) / static_cast<float>(NumUnreadyVolumes)
+			: RatioThreshold;
+		return Ratio < RatioThreshold
+			? ESpatialReadinessSimCallback_MidPhaseStrategy::PrioritizeMidPhases
+			: ESpatialReadinessSimCallback_MidPhaseStrategy::PrioritizeUnreadyVolumes;
+	}
+
+	// Just default back to prioritizing unready volumes if the cvar is set to something crazy
+	return ESpatialReadinessSimCallback_MidPhaseStrategy::PrioritizeUnreadyVolumes;
+}
+
+void FSpatialReadinessSimCallback::OnMidPhaseModification_PrioritizeUnreadyVolumes(FMidPhaseModifierAccessor& Accessor)
+{
+	SCOPE_CYCLE_COUNTER(STAT_SpatialReadiness_Physics_MidPhase_PrioritizeUnreadyVolumes)
 
 	// Go through every unready volume
 	for (FSingleParticlePhysicsProxy* UnreadyProxy : UnreadyVolumeParticles_PT)
@@ -600,7 +678,7 @@ void FSpatialReadinessSimCallback::OnMidPhaseModification_Internal(FMidPhaseModi
 #endif
 
 		{
-			SCOPE_CYCLE_COUNTER(STAT_SpatialReadiness_Physics_MidPhase_AccessorLoop)
+			//SCOPE_CYCLE_COUNTER(STAT_SpatialReadiness_Physics_MidPhase_AccessorLoop)
 
 			// Go through every mid-phase which involves this volume
 			for (FMidPhaseModifier& MidPhase : Accessor.GetMidPhases(UnreadyVolume))
@@ -640,7 +718,7 @@ void FSpatialReadinessSimCallback::OnMidPhaseModification_Internal(FMidPhaseModi
 #endif
 
 				// Add the particle to the list of particles to freeze
-				UnreadyRigidParticles_PT.FindOrAdd(RigidParticle) = { SyncTimeStamp, FreezeFrames };
+				UnreadyRigidParticles_PT.FindOrAdd(RigidParticle) = { SyncTimeStamp, GetFreezeFrames()};
 
 				// Disable the midphase
 				MidPhase.Disable();
@@ -649,6 +727,75 @@ void FSpatialReadinessSimCallback::OnMidPhaseModification_Internal(FMidPhaseModi
 	}
 }
 
+void FSpatialReadinessSimCallback::OnMidPhaseModification_PrioritizeMidPhases(FMidPhaseModifierAccessor& Accessor)
+{
+	SCOPE_CYCLE_COUNTER(STAT_SpatialReadiness_Physics_MidPhase_PrioritizeMidPhases)
+
+	Accessor.VisitMidPhases([this](FMidPhaseModifier& MidPhase)
+	{
+		FGeometryParticleHandle* UnreadyVolume;
+		FGeometryParticleHandle* GeometryParticle;
+		MidPhase.GetParticles(&UnreadyVolume, &GeometryParticle);
+
+		// Make sure we really have 2 particles
+		if (UnreadyVolume == nullptr || GeometryParticle == nullptr)
+		{
+			return;
+		}
+
+		// Make sure one of them is an unready volume particle
+		if (!IsUnreadyVolume_PT(*UnreadyVolume))
+		{
+			Swap(UnreadyVolume, GeometryParticle);
+			if (!IsUnreadyVolume_PT(*UnreadyVolume))
+			{
+				// Neither particle was an unready volume
+				return;
+			}
+		}
+
+		// Make sure the other one is a rigid - this should always pass
+		// because you can't have generated a midphase without at least
+		// one simulating particle.
+		FPBDRigidParticleHandle* RigidParticle = GeometryParticle->CastToRigidParticle();
+		if (RigidParticle == nullptr)
+		{
+			return;
+		}
+
+		FSingleParticlePhysicsProxy* RigidProxy = static_cast<FSingleParticlePhysicsProxy*>(RigidParticle->PhysicsProxy());
+		if (RigidProxy == nullptr)
+		{
+			return;
+		}
+
+		TWeakPtr<FProxyTimestampBase, ESPMode::ThreadSafe> SyncTimeStamp = RigidProxy->GetSyncTimestamp();
+		if (SyncTimeStamp.IsValid() == false)
+		{
+			return;
+		}
+
+		//SCOPE_CYCLE_COUNTER(STAT_SpatialReadiness_Physics_MidPhase_AccessorLoop)
+
+		// Add the particle to the list of particles to freeze
+		UnreadyRigidParticles_PT.FindOrAdd(RigidParticle) = { SyncTimeStamp, GetFreezeFrames()};
+
+		// Disable the midphase
+		MidPhase.Disable();
+	});
+}
+
+void FSpatialReadinessSimCallback::OnMidPhaseModification_Grid(FMidPhaseModifierAccessor& Accessor)
+{
+	SCOPE_CYCLE_COUNTER(STAT_SpatialReadiness_Physics_MidPhase_Grid)
+
+	//
+	// TODO: Everything!
+	// For now, just call back to prioritize Unready Volumes
+	//
+	OnMidPhaseModification_PrioritizeUnreadyVolumes(Accessor);
+}
+
 void FSpatialReadinessSimCallback::OnPreIntegrate_Internal()
 {
 	FreezeParticles_PT();
diff --git a/Engine/Plugins/Experimental/SpatialReadiness/Source/SpatialReadiness/Public/SpatialReadinessSimCallback.h b/Engine/Plugins/Experimental/SpatialReadiness/Source/SpatialReadiness/Public/SpatialReadinessSimCallback.h
--- a/Engine/Plugins/Experimental/SpatialReadiness/Source/SpatialReadiness/Public/SpatialReadinessSimCallback.h
+++ b/Engine/Plugins/Experimental/SpatialReadiness/Source/SpatialReadiness/Public/SpatialReadinessSimCallback.h
@@ -16,6 +16,15 @@ namespace Chaos
 	enum class EObjectStateType: int8;
 }
 
+enum ESpatialReadinessSimCallback_MidPhaseStrategy
+{
+	PrioritizeUnreadyVolumes = 0,
+	PrioritizeMidPhases,
+	Grid,
+	COUNT
+};
+
+
 struct FUnreadyVolumeData_GT
 {
 	FUnreadyVolumeData_GT(Chaos::FSingleParticlePhysicsProxy* InProxy, const FBox& InBounds, const FString& InDescription)
@@ -103,8 +112,16 @@ protected:
 	virtual void OnPostSolve_Internal() override;
 	/* end: TSimCallbackObject */
 
+	// Variants of the midphase modification function, with different
+	// performance characteristics
+	void OnMidPhaseModification_PrioritizeUnreadyVolumes(Chaos::FMidPhaseModifierAccessor& Accessor);
+	void OnMidPhaseModification_PrioritizeMidPhases(Chaos::FMidPhaseModifierAccessor& Accessor);
+	void OnMidPhaseModification_Grid(Chaos::FMidPhaseModifierAccessor& Accessor);
+
 	// Helpers
 	Chaos::FPBDRigidsEvolution* GetEvolution();
+	static int32 GetFreezeFrames();
+	ESpatialReadinessSimCallback_MidPhaseStrategy GetMidPhaseStrategy(const Chaos::FMidPhaseModifierAccessor& Accessor) const;
 
 	// Functions for freezing and unfreezing all particles in the
 	// UnreadyRigidParticles_PT set.
@@ -119,6 +136,9 @@ protected:
 	// TODO: Make const
 	bool QueryReadiness_PT(const Chaos::FAABB3& Bounds, TArray<const Chaos::FSingleParticlePhysicsProxy*>& OutVolumeProxies);
 
+	// True for particles which are in the unready volumes particles list
+	bool IsUnreadyVolume_PT(const Chaos::FGeometryParticleHandle& ParticleHandle) const;
+
 	// Keep a ref to the phys scene so we can add and remove particles
 	FPhysScene_Chaos& PhysicsScene;
 
diff --git a/Engine/Plugins/Experimental/SpatialReadiness/Source/SpatialReadiness/Public/SpatialReadinessStats.h b/Engine/Plugins/Experimental/SpatialReadiness/Source/SpatialReadiness/Public/SpatialReadinessStats.h
--- a/Engine/Plugins/Experimental/SpatialReadiness/Source/SpatialReadiness/Public/SpatialReadinessStats.h
+++ b/Engine/Plugins/Experimental/SpatialReadiness/Source/SpatialReadiness/Public/SpatialReadinessStats.h
@@ -16,7 +16,10 @@ DECLARE_CYCLE_STAT(TEXT("SpatialReadiness_Physics_DecrementUnreadyRigidParticles
 DECLARE_CYCLE_STAT(TEXT("SpatialReadiness_Physics_PreSimulate"), STAT_SpatialReadiness_Physics_PreSimulate, STATGROUP_SpatialReadiness);
 DECLARE_CYCLE_STAT(TEXT("SpatialReadiness_Physics_ParticlesRegistered"), STAT_SpatialReadiness_Physics_ParticlesRegistered, STATGROUP_SpatialReadiness);
 DECLARE_CYCLE_STAT(TEXT("SpatialReadiness_Physics_MidPhase"), STAT_SpatialReadiness_Physics_MidPhase, STATGROUP_SpatialReadiness);
-DECLARE_CYCLE_STAT(TEXT("SpatialReadiness_Physics_MidPhase_AccessorLoop"), STAT_SpatialReadiness_Physics_MidPhase_AccessorLoop, STATGROUP_SpatialReadinessVerbose);
+DECLARE_CYCLE_STAT(TEXT("SpatialReadiness_Physics_MidPhase_PrioritizeUnreadyVolumes"), STAT_SpatialReadiness_Physics_MidPhase_PrioritizeUnreadyVolumes, STATGROUP_SpatialReadiness);
+DECLARE_CYCLE_STAT(TEXT("SpatialReadiness_Physics_MidPhase_PrioritizeMidPhases"), STAT_SpatialReadiness_Physics_MidPhase_PrioritizeMidPhases, STATGROUP_SpatialReadiness);
+DECLARE_CYCLE_STAT(TEXT("SpatialReadiness_Physics_MidPhase_Grid"), STAT_SpatialReadiness_Physics_MidPhase_Grid, STATGROUP_SpatialReadiness);
+//DECLARE_CYCLE_STAT(TEXT("SpatialReadiness_Physics_MidPhase_AccessorLoop"), STAT_SpatialReadiness_Physics_MidPhase_AccessorLoop, STATGROUP_SpatialReadinessVerbose);
 DECLARE_CYCLE_STAT(TEXT("SpatialReadiness_Physics_ChaosVisualDebugger"), STAT_SpatialReadiness_Physics_ChaosVisualDebugger, STATGROUP_SpatialReadiness);
 DECLARE_CYCLE_STAT(TEXT("SpatialReadiness_Physics_DebugDraw"), STAT_SpatialReadiness_Physics_DebugDraw, STATGROUP_SpatialReadiness);
 
diff --git a/Engine/Source/Runtime/Experimental/Chaos/Private/Chaos/MidPhaseModification.cpp b/Engine/Source/Runtime/Experimental/Chaos/Private/Chaos/MidPhaseModification.cpp
--- a/Engine/Source/Runtime/Experimental/Chaos/Private/Chaos/MidPhaseModification.cpp
+++ b/Engine/Source/Runtime/Experimental/Chaos/Private/Chaos/MidPhaseModification.cpp
@@ -97,6 +97,16 @@ namespace Chaos
 		return FMidPhaseModifierParticleIterator(Accessor, Particle, Particle->ParticleCollisions().Num());
 	}
 
+	int32 FMidPhaseModifierAccessor::GetNumMidPhases() const
+	{
+		return Evolution.GetCollisionConstraints().GetConstraintAllocator().GetNumMidPhases();
+	}
+
+	int32 FMidPhaseModifierAccessor::GetNumMidPhases(const FGeometryParticleHandle* Particle) const
+	{
+		return Particle->ParticleCollisions().Num();
+	}
+
 	FMidPhaseModifierParticleRange FMidPhaseModifierAccessor::GetMidPhases(FGeometryParticleHandle* Particle)
 	{
 		return FMidPhaseModifierParticleRange(this, Particle);
diff --git a/Engine/Source/Runtime/Experimental/Chaos/Public/Chaos/Collision/CollisionConstraintAllocator.h b/Engine/Source/Runtime/Experimental/Chaos/Public/Chaos/Collision/CollisionConstraintAllocator.h
--- a/Engine/Source/Runtime/Experimental/Chaos/Public/Chaos/Collision/CollisionConstraintAllocator.h
+++ b/Engine/Source/Runtime/Experimental/Chaos/Public/Chaos/Collision/CollisionConstraintAllocator.h
@@ -517,6 +517,14 @@ namespace Chaos
 				}
 			}
 
+			/**
+			 * @brief Get the total number of midphases
+			 */
+			int32 GetNumMidPhases() const
+			{
+				return ParticlePairMidPhases.Num();
+			}
+
 			/**
 			 * Enable/Disable determinism (extra sorting steps after collision detection)
 			*/
diff --git a/Engine/Source/Runtime/Experimental/Chaos/Public/Chaos/MidPhaseModification.h b/Engine/Source/Runtime/Experimental/Chaos/Public/Chaos/MidPhaseModification.h
--- a/Engine/Source/Runtime/Experimental/Chaos/Public/Chaos/MidPhaseModification.h
+++ b/Engine/Source/Runtime/Experimental/Chaos/Public/Chaos/MidPhaseModification.h
@@ -160,6 +160,12 @@ namespace Chaos
 			: Evolution(InEvolution)
 		{ }
 
+		// Get the total number of mid phases for all particles
+		CHAOS_API int32 GetNumMidPhases() const;
+
+		// Get the number of mid phases involving a particular particle
+		CHAOS_API int32 GetNumMidPhases(const FGeometryParticleHandle* Particle) const;
+
 		// Get an object which allows for range iteration over the list of
 		// midphases for a particle
 		CHAOS_API FMidPhaseModifierParticleRange GetMidPhases(FGeometryParticleHandle* Particle);
