Updated spatial readiness to not use separate query/sim block filters and instead create separate shapes.
Added a cvar to disable the query shape as we think it isn't necessary.
#Synced-CL 48295927

diff --git a/Engine/Plugins/Experimental/SpatialReadiness/Source/SpatialReadiness/Private/SpatialReadiness/SpatialReadinessSimCallback.cpp b/Engine/Plugins/Experimental/SpatialReadiness/Source/SpatialReadiness/Private/SpatialReadiness/SpatialReadinessSimCallback.cpp
--- a/Engine/Plugins/Experimental/SpatialReadiness/Source/SpatialReadiness/Private/SpatialReadiness/SpatialReadinessSimCallback.cpp
+++ b/Engine/Plugins/Experimental/SpatialReadiness/Source/SpatialReadiness/Private/SpatialReadiness/SpatialReadinessSimCallback.cpp
@@ -76,6 +76,35 @@ TAutoConsoleVariable<float> CVarSpatialReadiness_MidPhaseStrategy_Dynamic_Thresh
 	TEXT("The ratio of mid phases to unready volumes below which we prioritize mid phases and above which we prioritize unready volumes."),
 	ECVF_Default);
 
+TAutoConsoleVariable<bool> CVarSpatialReadiness_CreateQueryShape(
+	TEXT("p.SpatialReadiness.CreateQueryShape"), false,
+	TEXT("When true, spatial readiness will create a sim and a query shape. When false, the query shape is skipped."),
+	ECVF_Default);
+
+// Callback returns Overlap for all sim shapes, None for everything else.
+struct FSpatialReadinessCollisionQueryFilterCallback : public FOverlapAllCollisionQueryFilterCallback
+{
+	virtual ~FSpatialReadinessCollisionQueryFilterCallback() = default;
+
+	ECollisionQueryHitType PreFilter(const Chaos::FPerShapeData& Shape)
+	{
+		if (!Shape.GetSimEnabled())
+		{
+			return ECollisionQueryHitType::None;
+		}
+		return ECollisionQueryHitType::Touch;
+	}
+
+	virtual ECollisionQueryHitType PreFilter(const FQueryFilterData& FilterData, const Chaos::FPerShapeData& Shape, const Chaos::FGeometryParticle& Actor) override
+	{
+		return PreFilter(Shape);
+	}
+	virtual ECollisionQueryHitType PreFilter(const FQueryFilterData& FilterData, const Chaos::FPerShapeData& Shape, const Chaos::FGeometryParticleHandle& Actor) override
+	{
+		return PreFilter(Shape);
+	}
+};
+
 FSpatialReadinessSimCallback::FSpatialReadinessSimCallback(FPhysScene_Chaos& InPhysicsScene)
 	: PhysicsScene(InPhysicsScene)
 	, UnreadyVolumeData_GT(256)
@@ -130,19 +159,42 @@ int32 FSpatialReadinessSimCallback::AddUnreadyVolume_GT(const FBox& Bounds, cons
 	Chaos::Filter::FShapeFilterBuilder FilterBuilder;
 	FilterBuilder.SetCollisionChannelIndex(ECollisionChannel::ECC_WorldDynamic);
 	FilterBuilder.SetOverlapChannelMask(TNumericLimits<uint64>::Max());
-	FilterBuilder.SetBlockChannelMasks(TNumericLimits<uint64>::Max(), 0);
+	FilterBuilder.SetBlockChannelMask(TNumericLimits<uint64>::Max());
 	FilterBuilder.SetFilterFlags(Chaos::EFilterFlags::SimpleCollision);
 	FilterBuilder.SetFilterFlags(Chaos::EFilterFlags::ComplexCollision);
 	FilterBuilder.SetFilterFlags(Chaos::EFilterFlags::CCD);
 	FilterBuilder.SetFilterFlags(Chaos::EFilterFlags::StaticShape);
-	const Chaos::Filter::FShapeFilterData ShapeFilterData = FilterBuilder.Build();
+	const Chaos::Filter::FShapeFilterData SimShapeFilterData = FilterBuilder.Build();
 
 	// Make the geometry
-	ParticleHandle.SetGeometry(BoxGeom);
+	const bool bCreateQueryShape = CVarSpatialReadiness_CreateQueryShape.GetValueOnAnyThread();
+
+	Chaos::FImplicitObjectPtr GeomPtr = BoxGeom;
+	if (bCreateQueryShape)
+	{
+		TArray<Chaos::FImplicitObjectPtr> Geometry{ BoxGeom, BoxGeom };
+		GeomPtr = MakeImplicitObjectPtr<Chaos::FImplicitObjectUnion>(MoveTemp(Geometry));
+	}
+	ParticleHandle.SetGeometry(GeomPtr);
+
 	ParticleHandle.SetShapeSimCollisionEnabled(0, true);
-	ParticleHandle.SetShapeQueryCollisionEnabled(0, true);
-	ParticleHandle.SetShapeFilterData(0, ShapeFilterData);
+	ParticleHandle.SetShapeQueryCollisionEnabled(0, false);
+	ParticleHandle.SetShapeFilterData(0, SimShapeFilterData);
 	ParticleHandle.ShapesArray()[0]->SetIsProbe(true);
+
+	if (bCreateQueryShape)
+	{
+		// The query filter is the same as sim, just with all blocking disabled.
+		Chaos::Filter::FShapeFilterBuilder QueryFilterBuilder(SimShapeFilterData);
+		QueryFilterBuilder.SetBlockChannelMask(0);
+		const Chaos::Filter::FShapeFilterData QueryShapeFilterData = QueryFilterBuilder.Build();
+
+		ParticleHandle.SetShapeSimCollisionEnabled(1, false);
+		ParticleHandle.SetShapeQueryCollisionEnabled(1, true);
+		ParticleHandle.SetShapeFilterData(1, QueryShapeFilterData);
+		ParticleHandle.ShapesArray()[1]->SetIsProbe(false);
+	}
+	
 #if CHAOS_DEBUG_NAME
 	ParticleHandle.SetDebugName(MakeShared<FString>(FString::Printf(TEXT("UnreadyVolume: %s"), *Description)));
 #endif
@@ -255,7 +307,7 @@ bool FSpatialReadinessSimCallback::QueryReadiness_GT(const FBox& Bounds, TArray<
 	FCollisionObjectQueryParams ObjectParams(static_cast<uint8>(Channel));
 	FCollisionResponseContainer ResponseContainer;
 	const Chaos::Filter::FQueryFilterData CollisionFilterData = CreateChaosQueryFilterData(ChannelBit, bComplex, ResponseContainer, QueryParams, ObjectParams, bMulti);
-	FOverlapAllCollisionQueryFilterCallback QueryCallback;
+	FSpatialReadinessCollisionQueryFilterCallback QueryCallback;
 	EQueryFlags QueryFlags = EQueryFlags::PreFilter;// | EQueryFlags::SkipNarrowPhase;
 
 	// Create the geometry needed for the query
