[Spatial Readiness] Ensure to disable MidPhase for Unready volumes

#Synced-CL 48903607

diff --git a/Engine/Plugins/Experimental/SpatialReadiness/Source/SpatialReadiness/Private/SpatialReadiness/SpatialReadinessSimCallback.cpp b/Engine/Plugins/Experimental/SpatialReadiness/Source/SpatialReadiness/Private/SpatialReadiness/SpatialReadinessSimCallback.cpp
--- a/Engine/Plugins/Experimental/SpatialReadiness/Source/SpatialReadiness/Private/SpatialReadiness/SpatialReadinessSimCallback.cpp
+++ b/Engine/Plugins/Experimental/SpatialReadiness/Source/SpatialReadiness/Private/SpatialReadiness/SpatialReadinessSimCallback.cpp
@@ -698,14 +698,18 @@ void FSpatialReadinessSimCallback::OnMidPhaseModification_PrioritizeUnreadyVolum
 	// Go through every unready volume
 	for (FSingleParticlePhysicsProxy* UnreadyProxy : UnreadyVolumeParticles_PT)
 	{
-		if (UnreadyProxy->GetMarkedDeleted())
+		FGeometryParticleHandle* UnreadyVolume = UnreadyProxy->GetHandle_LowLevel();
+		if (UnreadyVolume == nullptr)
 		{
 			continue;
 		}
-
-		FGeometryParticleHandle* UnreadyVolume = UnreadyProxy->GetHandle_LowLevel();
-		if (UnreadyVolume == nullptr)
+		if (UnreadyProxy->GetMarkedDeleted())
 		{
+			for (FMidPhaseModifier& MidPhase : Accessor.GetMidPhases(UnreadyVolume))
+			{
+				// Disable the midphase
+				MidPhase.Disable();
+			}
 			continue;
 		}
 
@@ -730,11 +734,12 @@ void FSpatialReadinessSimCallback::OnMidPhaseModification_PrioritizeUnreadyVolum
 #endif
 
 		{
-			//SCOPE_CYCLE_COUNTER(STAT_SpatialReadiness_Physics_MidPhase_AccessorLoop)
-
 			// Go through every mid-phase which involves this volume
 			for (FMidPhaseModifier& MidPhase : Accessor.GetMidPhases(UnreadyVolume))
 			{
+				// Disable the midphase because we know it is an Unready volume
+				MidPhase.Disable();
+
 				// Get the particle that is not the unready volume
 				FGeometryParticleHandle* GeometryParticle = MidPhase.GetOtherParticle(UnreadyVolume);
 				if (GeometryParticle == nullptr)
@@ -771,9 +776,6 @@ void FSpatialReadinessSimCallback::OnMidPhaseModification_PrioritizeUnreadyVolum
 
 				// Add the particle to the list of particles to freeze
 				UnreadyRigidParticles_PT.FindOrAdd(RigidParticle) = { SyncTimeStamp, GetFreezeFrames()};
-
-				// Disable the midphase
-				MidPhase.Disable();
 			}
 		}
 	}
@@ -806,6 +808,14 @@ void FSpatialReadinessSimCallback::OnMidPhaseModification_PrioritizeMidPhases(FM
 			}
 		}
 
+		// Disable the midphase for unready volumes
+		MidPhase.Disable();
+
+		if (UnreadyVolume->PhysicsProxy()->GetMarkedDeleted())
+		{
+			return;
+		}
+
 		// Make sure the other one is a rigid - this should always pass
 		// because you can't have generated a midphase without at least
 		// one simulating particle.
@@ -827,13 +837,9 @@ void FSpatialReadinessSimCallback::OnMidPhaseModification_PrioritizeMidPhases(FM
 			return;
 		}
 
-		//SCOPE_CYCLE_COUNTER(STAT_SpatialReadiness_Physics_MidPhase_AccessorLoop)
-
 		// Add the particle to the list of particles to freeze
 		UnreadyRigidParticles_PT.FindOrAdd(RigidParticle) = { SyncTimeStamp, GetFreezeFrames()};
 
-		// Disable the midphase
-		MidPhase.Disable();
 	});
 }
 
