Add missing physics scene lock to FSpatialReadinessSimCallback::AddUnreadyVolume_GT

diff --git a/Engine/Plugins/Experimental/SpatialReadiness/Source/SpatialReadiness/Private/SpatialReadiness/SpatialReadinessSimCallback.cpp b/Engine/Plugins/Experimental/SpatialReadiness/Source/SpatialReadiness/Private/SpatialReadiness/SpatialReadinessSimCallback.cpp
--- a/Engine/Plugins/Experimental/SpatialReadiness/Source/SpatialReadiness/Private/SpatialReadiness/SpatialReadinessSimCallback.cpp
+++ b/Engine/Plugins/Experimental/SpatialReadiness/Source/SpatialReadiness/Private/SpatialReadiness/SpatialReadinessSimCallback.cpp
@@ -200,8 +200,11 @@ int32 FSpatialReadinessSimCallback::AddUnreadyVolume_GT(const FBox& Bounds, cons
 #endif
 
 	// Add the new particle to the scene
-	TArray<FPhysicsActorHandle> Actors = { ParticleProxy };
-	PhysicsScene.AddActorsToScene_AssumesLocked(Actors);
+	{
+		FScopedSceneLock_Chaos SceneLock(&PhysicsScene, EPhysicsInterfaceScopedLockType::Write);
+		TArray<FPhysicsActorHandle> Actors = { ParticleProxy };
+		PhysicsScene.AddActorsToScene_AssumesLocked(Actors);
+	}
 
 	// Save the proxy in our list of GT particles.
 	//
