Avoid expensive includes in SpatialReadinessSubsystem.h and SpatialReadinessSimCallback.h

#Synced-CL 52185734

diff --git a/Engine/Plugins/Experimental/SpatialReadiness/Source/SpatialReadiness/Private/SpatialReadiness/SpatialReadinessDebug.cpp b/Engine/Plugins/Experimental/SpatialReadiness/Source/SpatialReadiness/Private/SpatialReadiness/SpatialReadinessDebug.cpp
--- a/Engine/Plugins/Experimental/SpatialReadiness/Source/SpatialReadiness/Private/SpatialReadiness/SpatialReadinessDebug.cpp
+++ b/Engine/Plugins/Experimental/SpatialReadiness/Source/SpatialReadiness/Private/SpatialReadiness/SpatialReadinessDebug.cpp
@@ -6,6 +6,7 @@
 #include "SpatialReadinessSimCallback.h"
 #include "PhysicsProxy/SingleParticlePhysicsProxy.h"
 #include "Physics/Experimental/PhysScene_Chaos.h"
+#include "PhysicsSolver.h"
 #include "Chaos/ParticleHandle.h"
 #include "Engine/World.h"
 
diff --git a/Engine/Plugins/Experimental/SpatialReadiness/Source/SpatialReadiness/Private/SpatialReadiness/SpatialReadinessSubsystem.cpp b/Engine/Plugins/Experimental/SpatialReadiness/Source/SpatialReadiness/Private/SpatialReadiness/SpatialReadinessSubsystem.cpp
--- a/Engine/Plugins/Experimental/SpatialReadiness/Source/SpatialReadiness/Private/SpatialReadiness/SpatialReadinessSubsystem.cpp
+++ b/Engine/Plugins/Experimental/SpatialReadiness/Source/SpatialReadiness/Private/SpatialReadiness/SpatialReadinessSubsystem.cpp
@@ -7,6 +7,7 @@
 #include "Engine/World.h"
 #include "Engine/Canvas.h"
 #include "Engine/Engine.h"
+#include "DrawDebugHelpers.h"
 #include "GameFramework/Pawn.h"
 #include "GameFramework/PlayerController.h"
 #include "PBDRigidsSolver.h"
diff --git a/Engine/Plugins/Experimental/SpatialReadiness/Source/SpatialReadiness/Public/SpatialReadinessSimCallback.h b/Engine/Plugins/Experimental/SpatialReadiness/Source/SpatialReadiness/Public/SpatialReadinessSimCallback.h
--- a/Engine/Plugins/Experimental/SpatialReadiness/Source/SpatialReadiness/Public/SpatialReadinessSimCallback.h
+++ b/Engine/Plugins/Experimental/SpatialReadiness/Source/SpatialReadiness/Public/SpatialReadinessSimCallback.h
@@ -7,9 +7,9 @@
 #include "Chaos/Particle/ObjectState.h"
 #include "Chaos/Framework/HashMappedArray.h"
 #include "SpatialReadinessVolume.h"
-#include "PhysicsSolver.h"
 
 class FPhysScene_Chaos;
+struct FProxyTimestampBase;
 namespace Chaos
 {
 	class FSingleParticlePhysicsProxy;
diff --git a/Engine/Plugins/Experimental/SpatialReadiness/Source/SpatialReadiness/Public/SpatialReadinessSubsystem.h b/Engine/Plugins/Experimental/SpatialReadiness/Source/SpatialReadiness/Public/SpatialReadinessSubsystem.h
--- a/Engine/Plugins/Experimental/SpatialReadiness/Source/SpatialReadiness/Public/SpatialReadinessSubsystem.h
+++ b/Engine/Plugins/Experimental/SpatialReadiness/Source/SpatialReadiness/Public/SpatialReadinessSubsystem.h
@@ -2,14 +2,15 @@
 #pragma once
 
 #include "CoreMinimal.h"
+#include "Chaos/Declares.h"
 #include "Subsystems/WorldSubsystem.h"
 #include "SpatialReadinessAPI.h"
 #include "SpatialReadinessSimCallback.h"
 #include "SpatialReadinessSignatures.h"
-#include "DrawDebugHelpers.h"
 #include "SpatialReadinessSubsystem.generated.h"
 
 class APlayerController;
+class UCanvas;
 
 UCLASS(MinimalAPI)
 class USpatialReadiness : public UWorldSubsystem
